fixed up incorrect module dependencies in Automation modules

diff --git a/Engine/Source/Developer/AutomationWindow/AutomationWindow.Build.cs b/Engine/Source/Developer/AutomationWindow/AutomationWindow.Build.cs
--- a/Engine/Source/Developer/AutomationWindow/AutomationWindow.Build.cs
+++ b/Engine/Source/Developer/AutomationWindow/AutomationWindow.Build.cs
@@ -7,24 +7,27 @@ namespace UnrealBuildTool.Rules
 		public AutomationWindow(TargetInfo Target)
 		{
 			PublicDependencyModuleNames.AddRange(
-				new string[]
-				{
+				new string[] {
 					"Core",
 					"AutomationController",
-					"SessionServices",
 				}
-				);
+			);
 
 			PrivateDependencyModuleNames.AddRange(
-				new string[]
-				{
+				new string[] {
 					"DesktopPlatform",
                     "InputCore",
-					"Messaging",
 					"Slate",
                     "EditorStyle",
 				}
-				);
+			);
+
+			PrivateIncludePathModuleNames.AddRange(
+				new string[] {
+					"Messaging",
+					"SessionServices",
+				}
+			);
 		}
 	}
 }
diff --git a/Engine/Source/Runtime/LauncherAutomatedService/LauncherAutomatedService.Build.cs b/Engine/Source/Runtime/LauncherAutomatedService/LauncherAutomatedService.Build.cs
--- a/Engine/Source/Runtime/LauncherAutomatedService/LauncherAutomatedService.Build.cs
+++ b/Engine/Source/Runtime/LauncherAutomatedService/LauncherAutomatedService.Build.cs
@@ -9,16 +9,21 @@ namespace UnrealBuildTool.Rules
 			PrivateIncludePaths.Add("Runtime/LauncherAutomatedService/Private");
 
 			PrivateDependencyModuleNames.AddRange(
-				new string[]
-				{
+				new string[] {
 					"AutomationController",
 					"Core",
 					"CoreUObject",
+					"TargetPlatform",
+				}
+			);
+
+			PrivateIncludePathModuleNames.AddRange(
+				new string[] {
 					"LauncherServices",
 					"Messaging",
 					"TargetDeviceServices",
 				}
-				);
+			);
 		}
 	}
 }
diff --git a/Engine/Source/Runtime/LauncherAutomatedService/Private/LauncherAutomatedServicePrivatePCH.h b/Engine/Source/Runtime/LauncherAutomatedService/Private/LauncherAutomatedServicePrivatePCH.h
--- a/Engine/Source/Runtime/LauncherAutomatedService/Private/LauncherAutomatedServicePrivatePCH.h
+++ b/Engine/Source/Runtime/LauncherAutomatedService/Private/LauncherAutomatedServicePrivatePCH.h
@@ -9,6 +9,7 @@
 #include "LauncherServices.h"
 #include "TargetDeviceServices.h"
 #include "Messaging.h"
+#include "TargetPlatform.h"
 #include "Ticker.h"
 #include "AutomationController.h"
 
