Automation:
Merging //depot/Partners/Zombie/UE4-iOS/... to //depot/UE4/...
Milestone 1 and 2 combined  (CL# 2047358 to 2055524)

 -Run external tools
 -Repeat tests N times
 -Visual difference for editor vs game tests
 -Static mesh UV test
 -Editor window screenshots
 -Automation test presets
 -Static mesh editor test

diff --git a/Engine/Config/BaseEngine.ini b/Engine/Config/BaseEngine.ini
--- a/Engine/Config/BaseEngine.ini
+++ b/Engine/Config/BaseEngine.ini
@@ -1843,6 +1843,22 @@ TestAllBlueprints=false
 +ReparentTest.ChildrenPackagePaths=/Game/ReparentingTestAssets/Children
 +ReparentTest.ParentsPackagePaths=/Game/ReparentingTestAssets/Parents
 
+[/Script/Engine.AutomationTestSettings]
++EditorTestModules=StaticMeshEditor
+
+[AutomationTesting.StaticMeshEditorTest]
++EditorViewButtons=Wireframe
++EditorViewButtons=Verts
++EditorViewButtons=Grid
++EditorViewButtons=Bounds
++EditorViewButtons=Collision
++EditorViewButtons=Show Pivot
++EditorViewButtons=Normals
++EditorViewButtons=Tangents
++EditorViewButtons=Binormals
++EditorViewButtons=UV
+EditorViewButtonsObject=EditorCylinder
+
 [/Script/Engine.NavigationSystem]
 bAutoCreateNavigationData=true
 bBuildNavigationAtRuntime=true
diff --git a/Engine/Source/Developer/AutomationController/Private/AutomationControllerManager.h b/Engine/Source/Developer/AutomationController/Private/AutomationControllerManager.h
--- a/Engine/Source/Developer/AutomationController/Private/AutomationControllerManager.h
+++ b/Engine/Source/Developer/AutomationController/Private/AutomationControllerManager.h
@@ -31,6 +31,16 @@ public:
 
 	virtual void Tick() OVERRIDE;
 
+	virtual void SetNumPasses(const int32 InNumPasses) OVERRIDE
+	{
+		NumTestPasses = InNumPasses;
+	}
+
+	virtual int32 GetNumPasses() OVERRIDE
+	{
+		return NumTestPasses;
+	}
+
 	virtual void SetFilter( TSharedPtr< AutomationFilterCollection > InFilter ) OVERRIDE
 	{
 		ReportManager.SetFilter( InFilter );
@@ -71,6 +81,16 @@ public:
 		return ReportManager.GetEnabledTestsNum();
 	}
 
+	virtual void GetEnabledTestNames(TArray<FString>& OutEnabledTestNames) const OVERRIDE
+	{
+		ReportManager.GetEnabledTestNames(OutEnabledTestNames);
+	}
+
+	virtual void SetEnabledTests(const TArray<FString>& EnabledTests) OVERRIDE
+	{
+		ReportManager.SetEnabledTests(EnabledTests);
+	}
+
 	virtual EAutomationControllerModuleState::Type GetTestState( ) const OVERRIDE
 	{
 		return AutomationTestState;
@@ -314,7 +334,11 @@ private:
 	/** A array of running tests */
 	TArray< FTestRunningInfo > TestRunningArray;
 
+	/** The number of test passes to perform */
+	int32 NumTestPasses;
 
+	/** The current test pass we are on */
+	int32 CurrentTestPass;
 private:
 
 	// Holds a delegate that is invoked when the controller shuts down.
diff --git a/Engine/Source/Developer/AutomationController/Private/AutomationControllerManger.cpp b/Engine/Source/Developer/AutomationController/Private/AutomationControllerManger.cpp
--- a/Engine/Source/Developer/AutomationController/Private/AutomationControllerManger.cpp
+++ b/Engine/Source/Developer/AutomationController/Private/AutomationControllerManger.cpp
@@ -54,6 +54,8 @@ void FAutomationControllerManager::RequestTests()
 void FAutomationControllerManager::RunTests( const bool bInIsLocalSession )
 {
 	ExecutionCount++;
+	CurrentTestPass = 0;
+	ReportManager.SetCurrentTestPass(CurrentTestPass);
 	ClusterDistributionMask = 0;
 	bTestResultsAvailable = false;
 	TestRunningArray.Empty();
@@ -64,7 +66,7 @@ void FAutomationControllerManager::RunTests( const bool bInIsLocalSession )
 	CheckTestTimer = 0.f;
 
 	//reset all tests
-	ReportManager.ResetForExecution();
+	ReportManager.ResetForExecution(NumTestPasses);
 
 	for (int32 ClusterIndex = 0; ClusterIndex < DeviceClusterManager.GetNumClusters(); ++ClusterIndex)
 	{
@@ -182,11 +184,11 @@ void FAutomationControllerManager::ExecuteNextTask( int32 ClusterIndex, OUT bool
 		if( !DeviceClusterManager.GetTest( ClusterIndex, DeviceIndex ).IsValid() && DeviceClusterManager.DeviceEnabled( ClusterIndex, DeviceIndex ) )
 		{
 			// Get the next test that should be worked on
-			TSharedPtr< IAutomationReport > NextTest = ReportManager.GetNextReportToExecute( bAllTestsCompleted, ClusterIndex, NumDevicesInCluster );
+			TSharedPtr< IAutomationReport > NextTest = ReportManager.GetNextReportToExecute( bAllTestsCompleted, ClusterIndex,CurrentTestPass, NumDevicesInCluster );
 			if( NextTest.IsValid() )
 			{
 				// Get the status of the test
-				EAutomationState::Type TestState = NextTest->GetState( ClusterIndex );
+				EAutomationState::Type TestState = NextTest->GetState( ClusterIndex,CurrentTestPass );
 				if( TestState == EAutomationState::NotRun )
 				{
 					// Reserve this device for the test
@@ -203,7 +205,7 @@ void FAutomationControllerManager::ExecuteNextTask( int32 ClusterIndex, OUT bool
 							FAutomationTestResults TestResults;
 							TestResults.State = EAutomationState::InProcess;
 							TestResults.GameInstance = DeviceClusterManager.GetClusterDeviceName( ClusterIndex, DeviceIndex );
-							NextTest->SetResults( ClusterIndex, TestResults );
+							NextTest->SetResults( ClusterIndex,CurrentTestPass, TestResults );
 							NextTest->ResetNetworkCommandResponses();
 
 							// Mark the device as busy
@@ -245,12 +247,19 @@ void FAutomationControllerManager::ExecuteNextTask( int32 ClusterIndex, OUT bool
 				TestResults.GameInstance = DeviceClusterManager.GetClusterDeviceName(ClusterIndex, 0);
 				TestResults.Warnings.Append( AutomationsWarnings );
 
-				CurrentTest->SetResults( ClusterIndex, TestResults );
+				CurrentTest->SetResults( ClusterIndex,CurrentTestPass, TestResults );
 				DeviceClusterManager.ResetAllDevicesRunningTest( ClusterIndex, CurrentTest );
 			}
-
 		}
 	}
+
+	//Check to see if we finished a pass
+	if( bAllTestsCompleted && CurrentTestPass < NumTestPasses - 1 )
+	{
+		CurrentTestPass++;
+		ReportManager.SetCurrentTestPass(CurrentTestPass);
+		bAllTestsCompleted = false;
+	}
 }
 
 
@@ -276,6 +285,7 @@ void FAutomationControllerManager::Startup()
 	bVisualCommandletFilterOn = false;
 
 	NumOfTestsToReceive = 0;
+	NumTestPasses = 1;
 }
 
 
@@ -307,7 +317,7 @@ void FAutomationControllerManager::SetTestNames( const FMessageAddress& Automati
 		for( int32 TestIndex = 0; TestIndex < TestInfo.Num(); ++TestIndex )
 		{
 			// Ensure our test exists. If not, add it
-			ReportManager.EnsureReportExists( TestInfo[TestIndex], DeviceClusterIndex );
+			ReportManager.EnsureReportExists( TestInfo[TestIndex], DeviceClusterIndex, NumTestPasses );
 		}
 
 		// Clear any intermediate data we had associated with the tests whilst building the full list of tests
@@ -366,7 +376,7 @@ void FAutomationControllerManager::CheckChildResult(TSharedPtr<IAutomationReport
 	{
 		for (int32 ClusterIndex = 0; ClusterIndex < GetNumDeviceClusters(); ++ClusterIndex)
 		{
-			FAutomationTestResults TestResults = InReport->GetResults( ClusterIndex );
+			FAutomationTestResults TestResults = InReport->GetResults( ClusterIndex,CurrentTestPass );
 
 			if (TestResults.Errors.Num())
 			{
@@ -456,7 +466,7 @@ void FAutomationControllerManager::UpdateTests( )
 				TestResults.GameInstance = DeviceClusterManager.GetClusterDeviceName(ClusterIndex, DeviceIndex);
 
 				// Set the results
-				Report->SetResults(ClusterIndex, TestResults );
+				Report->SetResults(ClusterIndex,CurrentTestPass, TestResults );
 				bTestResultsAvailable = true;
 
 				// Disable the device in the cluster so it is not used again
@@ -626,7 +636,7 @@ void FAutomationControllerManager::HandleRunTestsReplyMessage( const FAutomation
 		TSharedPtr <IAutomationReport> Report = DeviceClusterManager.GetTest(ClusterIndex, DeviceIndex);
 		check(Report.IsValid());
 
-		Report->SetResults(ClusterIndex, TestResults);
+		Report->SetResults(ClusterIndex,CurrentTestPass, TestResults);
 
 		// Device is now good to go
 		DeviceClusterManager.SetTest(ClusterIndex, DeviceIndex, NULL);
diff --git a/Engine/Source/Developer/AutomationController/Private/AutomationReport.cpp b/Engine/Source/Developer/AutomationController/Private/AutomationReport.cpp
--- a/Engine/Source/Developer/AutomationController/Private/AutomationReport.cpp
+++ b/Engine/Source/Developer/AutomationController/Private/AutomationReport.cpp
@@ -79,6 +79,60 @@ int32 FAutomationReport::GetTotalNumChildren() const
 	return Total;
 }
 
+void FAutomationReport::GetEnabledTestNames(TArray<FString>& OutEnabledTestNames, FString CurrentPath) const
+{
+	//if this is a leaf and this test is enabled
+	if ((ChildReports.Num() == 0) && IsEnabled())
+	{
+		const FString FullTestName = CurrentPath.Len() > 0 ? CurrentPath.AppendChar(TCHAR('.')) + TestInfo.GetDisplayName() : TestInfo.GetDisplayName();
+		OutEnabledTestNames.Add(FullTestName);
+	}
+	else
+	{
+		if( !CurrentPath.IsEmpty() )
+		{
+			CurrentPath += TEXT(".");
+		}
+		CurrentPath += TestInfo.GetDisplayName();
+		//recurse through the hierarchy
+		for (int32 ChildIndex = 0; ChildIndex < ChildReports.Num(); ++ChildIndex)
+		{
+			ChildReports[ChildIndex]->GetEnabledTestNames(OutEnabledTestNames,CurrentPath);
+		}
+	}
+	return;
+}
+
+void FAutomationReport::SetEnabledTests(const TArray<FString>& EnabledTests, FString CurrentPath)
+{
+	if (ChildReports.Num() == 0)
+	{
+		//Find of the full name of this test and see if it is in our list
+		const FString FullTestName = CurrentPath.Len() > 0 ? CurrentPath.AppendChar(TCHAR('.')) + TestInfo.GetDisplayName() : TestInfo.GetDisplayName();
+		const bool bNewEnabled = EnabledTests.Contains(FullTestName);
+		SetEnabled(bNewEnabled);
+	}
+	else
+	{
+		if( !CurrentPath.IsEmpty() )
+		{
+			CurrentPath += TEXT(".");
+		}
+		CurrentPath += TestInfo.GetDisplayName();
+
+		//recurse through the hierarchy
+		for (int32 ChildIndex = 0; ChildIndex < ChildReports.Num(); ++ChildIndex)
+		{
+			ChildReports[ChildIndex]->SetEnabledTests(EnabledTests,CurrentPath);
+		}
+
+		//Parent nodes should be checked if all of its children are
+		const int32 TotalNumChildern = GetTotalNumChildren();
+		const int32 EnabledChildren = GetEnabledTestsNum();
+		bEnabled = (TotalNumChildern == EnabledChildren);
+	}
+}
+
 /** Recursively gets the number of enabled tests */
 int32 FAutomationReport::GetEnabledTestsNum() const
 {
@@ -123,7 +177,10 @@ void FAutomationReport::SetSupport(const int32 ClusterIndex)
 	//ensure there is enough room in the array for status per platform
 	for (int32 i = 0; i <= ClusterIndex; ++i)
 	{
-		Results.Add( FAutomationTestResults() );
+		//Make sure we have enough results for a single pass
+		TArray<FAutomationTestResults> AutomationTestResult;
+		AutomationTestResult.Add( FAutomationTestResults() );
+		Results.Add( AutomationTestResult );
 	}
 }
 
@@ -211,54 +268,71 @@ TArray<TSharedPtr<IAutomationReport> >& FAutomationReport::GetChildReports()
 	return ChildReports;
 }
 
-void FAutomationReport::ResetForExecution()
+void FAutomationReport::ResetForExecution(const int32 NumTestPasses)
 {
 	TestInfo.ResetNumDevicesRunningTest();
 
 	//if this test is enabled
 	if (IsEnabled())
 	{
-		for (int32 i = 0; i < Results.Num(); ++i)
+		for (int32 ClusterIndex = 0; ClusterIndex < Results.Num(); ++ClusterIndex)
 		{
-			//reset all stats
-			Results[i].State = EAutomationState::NotRun;
-			Results[i].Warnings.Empty();
-			Results[i].Errors.Empty();
+			//Make sure we have enough results
+			if( NumTestPasses > Results[ClusterIndex].Num() )
+			{
+				for(int32 PassCount = Results[ClusterIndex].Num(); PassCount < NumTestPasses; ++PassCount)
+				{
+					Results[ClusterIndex].Add( FAutomationTestResults() );
+				}
+			}
+			else if( NumTestPasses < Results[ClusterIndex].Num() )
+			{
+				Results[ClusterIndex].RemoveAt(NumTestPasses, Results[ClusterIndex].Num() - NumTestPasses);
+			}
+
+			for( int32 PassIndex = 0; PassIndex < Results[ClusterIndex].Num(); ++PassIndex)
+			{
+				//reset all stats
+				Results[ClusterIndex][PassIndex].State = EAutomationState::NotRun;
+				Results[ClusterIndex][PassIndex].Warnings.Empty();
+				Results[ClusterIndex][PassIndex].Errors.Empty();
+			}
 		}
 	}
 	//recurse to children
 	for (int32 ChildIndex = 0; ChildIndex < ChildReports.Num(); ++ChildIndex)
 	{
-		ChildReports[ChildIndex]->ResetForExecution();
+		ChildReports[ChildIndex]->ResetForExecution(NumTestPasses);
 	}
 }
 
 
-void FAutomationReport::SetResults( const int32 ClusterIndex, const FAutomationTestResults& InResults )
+void FAutomationReport::SetResults( const int32 ClusterIndex, const int32 PassIndex, const FAutomationTestResults& InResults )
 {
 	//verify this is a platform this test is aware of
 	check((ClusterIndex >= 0) && (ClusterIndex < Results.Num()));
+	check((PassIndex >= 0) && (PassIndex < Results[ClusterIndex].Num()));
 
 	if( InResults.State == EAutomationState::InProcess )
 	{
 		TestInfo.InformOfNewDeviceRunningTest();
 	}
 
-	Results[ ClusterIndex ] = InResults;
+	Results[ ClusterIndex ][ PassIndex ] = InResults;
 	// Add an error report if none was received
 	if ( InResults.State == EAutomationState::Fail && InResults.Errors.Num() == 0 && InResults.Warnings.Num() == 0 )
 	{
-		Results[ClusterIndex].Errors.Add( "No Report Generated" );
+		Results[ClusterIndex][PassIndex].Errors.Add( "No Report Generated" );
 	}
 }
 
 
-void FAutomationReport::GetCompletionStatus(const int32 ClusterIndex, FAutomationCompleteState& OutCompletionState)
+void FAutomationReport::GetCompletionStatus(const int32 ClusterIndex, const int32 PassIndex, FAutomationCompleteState& OutCompletionState)
 {
 	//if this test is enabled and a leaf test
 	if (IsSupported(ClusterIndex) && (ChildReports.Num()==0))
 	{
-		EAutomationState::Type CurrentState = Results[ClusterIndex].State;
+		EAutomationState::Type CurrentState = Results[ClusterIndex][PassIndex].State;
 		//Enabled and In-Process
 		if (IsEnabled())
 		{
@@ -270,7 +344,7 @@ void FAutomationReport::GetCompletionStatus(const int32 ClusterIndex, FAutomatio
 		}
 
 		//Warnings
-		if (Results[ClusterIndex].Warnings.Num() > 0)
+		if (Results[ClusterIndex][PassIndex].Warnings.Num() > 0)
 		{
 			IsEnabled() ? OutCompletionState.NumEnabledTestsWarnings++ : OutCompletionState.NumDisabledTestsWarnings++;
 		}
@@ -292,32 +366,49 @@ void FAutomationReport::GetCompletionStatus(const int32 ClusterIndex, FAutomatio
 	//recurse to children
 	for (int32 ChildIndex = 0; ChildIndex < ChildReports.Num(); ++ChildIndex)
 	{
-		ChildReports[ChildIndex]->GetCompletionStatus(ClusterIndex, OutCompletionState);
+		ChildReports[ChildIndex]->GetCompletionStatus(ClusterIndex,PassIndex, OutCompletionState);
 	}
 }
 
 
-EAutomationState::Type FAutomationReport::GetState(const int32 ClusterIndex) const
+EAutomationState::Type FAutomationReport::GetState(const int32 ClusterIndex, const int32 PassIndex) const
 {
-	if ((ClusterIndex >= 0) && (ClusterIndex < Results.Num()))
+	if ((ClusterIndex >= 0) && (ClusterIndex < Results.Num()) &&
+		(PassIndex >= 0) && (PassIndex < Results[ClusterIndex].Num()))
 	{
-		return Results[ClusterIndex].State;
+		return Results[ClusterIndex][PassIndex].State;
 	}
 	return EAutomationState::NotRun;
 }
 
 
-const FAutomationTestResults& FAutomationReport::GetResults( const int32 ClusterIndex ) 
+const FAutomationTestResults& FAutomationReport::GetResults( const int32 ClusterIndex, const int32 PassIndex ) 
+{
+	return Results[ClusterIndex][PassIndex];
+}
+
+const int32 FAutomationReport::GetNumResults( const int32 ClusterIndex )
 {
-	return Results[ClusterIndex];
+	return Results[ClusterIndex].Num();
+}
+
+const int32 FAutomationReport::GetCurrentPassIndex( const int32 ClusterIndex )
+{
+	int32 PassIndex = 1;
+	for(; PassIndex < Results[ClusterIndex].Num(); ++PassIndex )
+	{
+		if( Results[ClusterIndex][PassIndex].State == EAutomationState::NotRun )
+			break;
+	}
+	return PassIndex - 1;
 }
 
 FString FAutomationReport::GetGameInstanceName( const int32 ClusterIndex )
 {
-	return Results[ClusterIndex].GameInstance;
+	return Results[ClusterIndex][0].GameInstance;
 }
 
-TSharedPtr<IAutomationReport> FAutomationReport::EnsureReportExists(FAutomationTestInfo& InTestInfo, const int32 ClusterIndex)
+TSharedPtr<IAutomationReport> FAutomationReport::EnsureReportExists(FAutomationTestInfo& InTestInfo, const int32 ClusterIndex, const int32 NumPasses)
 {
 	//Split New Test Name by the first "." found
 	FString NameToMatch = InTestInfo.GetDisplayName();
@@ -378,14 +469,14 @@ TSharedPtr<IAutomationReport> FAutomationReport::EnsureReportExists(FAutomationT
 	else
 	{
 		//recurse to add to the proper layer
-		FoundTest = MatchTest->EnsureReportExists(InTestInfo, ClusterIndex);
+		FoundTest = MatchTest->EnsureReportExists(InTestInfo, ClusterIndex, NumPasses);
 	}
 
 	return FoundTest;
 }
 
 
-TSharedPtr<IAutomationReport> FAutomationReport::GetNextReportToExecute(bool& bOutAllTestsComplete, const int32 ClusterIndex, const int32 NumDevicesInCluster)
+TSharedPtr<IAutomationReport> FAutomationReport::GetNextReportToExecute(bool& bOutAllTestsComplete, const int32 ClusterIndex, const int32 PassIndex, const int32 NumDevicesInCluster)
 {
 	TSharedPtr<IAutomationReport> NextReport;
 	//if this is not a leaf node
@@ -393,7 +484,7 @@ TSharedPtr<IAutomationReport> FAutomationReport::GetNextReportToExecute(bool& bO
 	{
 		for (int32 ReportIndex = 0; ReportIndex < ChildReports.Num(); ++ReportIndex)
 		{
-			NextReport = ChildReports[ReportIndex]->GetNextReportToExecute(bOutAllTestsComplete, ClusterIndex, NumDevicesInCluster);
+			NextReport = ChildReports[ReportIndex]->GetNextReportToExecute(bOutAllTestsComplete, ClusterIndex, PassIndex, NumDevicesInCluster);
 			//if we found one, return it
 			if (NextReport.IsValid())
 			{
@@ -406,7 +497,7 @@ TSharedPtr<IAutomationReport> FAutomationReport::GetNextReportToExecute(bool& bO
 		//consider self
 		if (IsEnabled() && IsSupported(ClusterIndex))
 		{
-			EAutomationState::Type TestState = GetState(ClusterIndex);
+			EAutomationState::Type TestState = GetState(ClusterIndex,PassIndex);
 			//if any enabled test hasn't been run yet or is in process
 			if ((TestState != EAutomationState::Success) && (TestState != EAutomationState::Fail) && (TestState != EAutomationState::NotEnoughParticipants))
 			{
@@ -427,12 +518,15 @@ const bool FAutomationReport::HasErrors()
 	bool bHasErrors = false;
 	for (int32 ClusterIndex = 0; ClusterIndex < Results.Num(); ++ClusterIndex )
 	{
-		//if we want tests with errors and this test had them OR we want tests warnings and this test had them
-		if( Results[ ClusterIndex ].Errors.Num() ) 
+		for( int32 PassIndex = 0; PassIndex < Results[ClusterIndex].Num(); ++PassIndex)
 		{
-			//mark this test as having passed the results filter
-			bHasErrors = true;
-			break;
+			//if we want tests with errors and this test had them OR we want tests warnings and this test had them
+			if( Results[ ClusterIndex ][ PassIndex ].Errors.Num() ) 
+			{
+				//mark this test as having passed the results filter
+				bHasErrors = true;
+				break;
+			}
 		}
 	}
 	return bHasErrors;
@@ -443,12 +537,15 @@ const bool FAutomationReport::HasWarnings()
 	bool bHasWarnings = false;
 	for (int32 ClusterIndex = 0; ClusterIndex < Results.Num(); ++ClusterIndex )
 	{
-		//if we want tests with errors and this test had them OR we want tests warnings and this test had them
-		if( Results[ ClusterIndex ].Warnings.Num() ) 
+		for( int32 PassIndex = 0; PassIndex < Results[ClusterIndex].Num(); ++PassIndex)
 		{
-			//mark this test as having passed the results filter
-			bHasWarnings = true;
-			break;
+			//if we want tests with errors and this test had them OR we want tests warnings and this test had them
+			if( Results[ ClusterIndex ][ PassIndex ].Warnings.Num() ) 
+			{
+				//mark this test as having passed the results filter
+				bHasWarnings = true;
+				break;
+			}
 		}
 	}
 	return bHasWarnings;
@@ -486,12 +583,15 @@ const bool FAutomationReport::GetDurationRange(float& OutMinTime, float& OutMaxT
 
 	for (int32 ClusterIndex = 0; ClusterIndex < Results.Num(); ++ClusterIndex )
 	{
-		//if we want tests with errors and this test had them OR we want tests warnings and this test had them
-		if( Results[ClusterIndex].State == EAutomationState::Success)
+		for( int32 PassIndex = 0; PassIndex < Results[ClusterIndex].Num(); ++PassIndex)
 		{
-			OutMinTime = FMath::Min(OutMinTime, Results[ClusterIndex].Duration );
-			OutMaxTime = FMath::Max(OutMaxTime, Results[ClusterIndex].Duration );
-			bAnyResultsFound = true;
+			//if we want tests with errors and this test had them OR we want tests warnings and this test had them
+			if( Results[ClusterIndex][PassIndex].State == EAutomationState::Success)
+			{
+				OutMinTime = FMath::Min(OutMinTime, Results[ClusterIndex][PassIndex].Duration );
+				OutMaxTime = FMath::Max(OutMaxTime, Results[ClusterIndex][PassIndex].Duration );
+				bAnyResultsFound = true;
+			}
 		}
 	}
 	return bAnyResultsFound;
@@ -541,9 +641,12 @@ void FAutomationReport::StopRunningTest()
 	{
 		for( int32 ResultsIndex = 0; ResultsIndex < Results.Num(); ++ResultsIndex )
 		{
-			if( Results[ResultsIndex].State == EAutomationState::InProcess )
+			for( int32 PassIndex = 0; PassIndex < Results[ResultsIndex].Num(); ++PassIndex)
 			{
-				Results[ResultsIndex].State = EAutomationState::NotRun;
+				if( Results[ResultsIndex][PassIndex].State == EAutomationState::InProcess )
+				{
+					Results[ResultsIndex][PassIndex].State = EAutomationState::NotRun;
+				}
 			}
 		}
 	}
diff --git a/Engine/Source/Developer/AutomationController/Private/AutomationReport.h b/Engine/Source/Developer/AutomationController/Private/AutomationReport.h
--- a/Engine/Source/Developer/AutomationController/Private/AutomationReport.h
+++ b/Engine/Source/Developer/AutomationController/Private/AutomationReport.h
@@ -48,6 +48,10 @@ public:
 	virtual int32 GetTotalNumChildren() const OVERRIDE;
 	/** Recursively gets the number of enabled tests */
 	virtual int32 GetEnabledTestsNum() const OVERRIDE;
+	/** Recursively gets the names of enabled tests */
+	virtual void GetEnabledTestNames(TArray<FString>& OutEnabledTestNames, FString CurrentPath) const OVERRIDE;
+	/** Recursively sets the enabled flag based on if the current test exists in the enabled tests array */
+	virtual void SetEnabledTests(const TArray<FString>& EnabledTests, FString CurrentPath) OVERRIDE;
 
 	/** Return if this test should be executed */
 	virtual bool IsEnabled() const OVERRIDE;
@@ -82,37 +86,59 @@ public:
 	/** Returns the array of child reports */
 	virtual TArray<TSharedPtr<IAutomationReport> >& GetChildReports() OVERRIDE;
 
-	/** Recursively resets the report to "needs to be run", clears cached warnings and errors */
-	virtual void ResetForExecution() OVERRIDE;
+	/** 
+	 * Recursively resets the report to "needs to be run", clears cached warnings and errors
+	 * @param NumTestPasses - The number of test passes so we know how many results to create
+	 */
+	virtual void ResetForExecution(const int32 NumTestPasses) OVERRIDE;
+
 	/**
 	 * Sets the results of the test for use by the UI
 	 * @param ClusterIndex - Index of the platform reporting the results of this test.  See AutomationDeviceClusterManager
+	 * @param PassIndex - Which test pass these results are for
 	 * @param InResults - The new set of results
 	 */
-	virtual void SetResults( const int32 ClusterIndex, const FAutomationTestResults& InResults ) OVERRIDE;
+	virtual void SetResults( const int32 ClusterIndex, const int32 PassIndex, const FAutomationTestResults& InResults ) OVERRIDE;
 
 	/**
 	 * Returns completion statistics for this branch of the testing hierarchy
 	 * @param ClusterIndex - Index of the platform reporting the results of this test.  See AutomationDeviceClusterManager
+	 * @param PassIndex - Which test pass to get the status of
 	 * @param OutCompletionState - Collection structure for execution statistics
 	 */
-	virtual void GetCompletionStatus(const int32 ClusterIndex, FAutomationCompleteState& OutCompletionState) OVERRIDE;
+	virtual void GetCompletionStatus(const int32 ClusterIndex, const int32 PassIndex, FAutomationCompleteState& OutCompletionState) OVERRIDE;
 
 	/**
 	 * Returns the state of the test (not run, in process, success, failure)
 	 * @param ClusterIndex - Index of the platform reporting the results of this test.  See AutomationDeviceClusterManager
+	 * @param PassIndex - Which test pass to get the state of
 	 * @return the current state of the test
 	 */
-	virtual EAutomationState::Type GetState(const int32 ClusterIndex) const OVERRIDE;
+	virtual EAutomationState::Type GetState(const int32 ClusterIndex, const int32 PassIndex) const OVERRIDE;
 
 	/**
 	 * Gets a copy of errors and warnings that were found
 	 *
 	 * @param ClusterIndex - Index of the platform we are requesting test results for.
+	 * @param PassIndex - Index of the test pass to get the results for.
 	 *
 	 * @return The collection of results for the given cluster index
 	 */
-	virtual const FAutomationTestResults& GetResults( const int32 ClusterIndex ) OVERRIDE;
+	virtual const FAutomationTestResults& GetResults( const int32 ClusterIndex, const int32 PassIndex ) OVERRIDE;
+
+	/**
+	 * Gets the number of available test results for a given cluster
+	 * @param ClusterIndex - Index of the platform .
+	 * @return The number of results available for this cluster
+	 */
+	virtual const int32 GetNumResults( const int32 ClusterIndex ) OVERRIDE;
+
+	/**
+	 * Finds the current pass by looking at the current state
+	 * @param ClusterIndex - Index of the platform.
+	 * @return The current pass index
+	 */
+	virtual const int32 GetCurrentPassIndex( const int32 ClusterIndex ) OVERRIDE;
 
 	/**
 	 * Gets the name of the instance that ran the test
@@ -126,19 +152,21 @@ public:
 	 * If NewTestName is Editor.Maps.Loadall.TestName, this will create nodes for Editor, Maps, Loadall, and then a leaf node for the testname with the associated commandline
 	 * @param TestInfo - Structure containing all the test info
 	 * @param ClusterIndex - Index of the platform reporting the results of this test.  See AutomationDeviceClusterManager
+	 * @param NumPasses - The number of passes we are going to perform.  Used to make sure we have enough results.
 	 * @return - the automation report
 	 */
-	virtual TSharedPtr<IAutomationReport> EnsureReportExists(FAutomationTestInfo& TestInfo, const int32 ClusterIndex) OVERRIDE;
+	virtual TSharedPtr<IAutomationReport> EnsureReportExists(FAutomationTestInfo& TestInfo, const int32 ClusterIndex, const int32 NumPasses) OVERRIDE;
 
 	/**
 	 * Returns the next test in the hierarchy to run
 	 * @param bOutAllTestsComplete - Whether or not all enabled tests have completed execution for this platform (cluster index)
 	 * @param ClusterIndex - Index of the platform reporting the results of this test.  See AutomationDeviceClusterManager
+	 * @param PassIndex - Which test pass we are currently on
 	 * @param NumDevicesInCluster - The number of devices which are in this cluster
 	 *
 	 * @return The next report
 	 */
-	virtual TSharedPtr<IAutomationReport> GetNextReportToExecute(bool& bOutAllTestsComplete, const int32 ClusterIndex, const int32 NumDevicesInCluster) OVERRIDE;
+	virtual TSharedPtr<IAutomationReport> GetNextReportToExecute(bool& bOutAllTestsComplete, const int32 ClusterIndex, const int32 PassIndex, const int32 NumDevicesInCluster) OVERRIDE;
 
 	/** 
 	 * Have any results got errors
@@ -234,7 +262,7 @@ private:
 	TArray<TSharedPtr<IAutomationReport> >FilteredChildReports;
 
 	/** Results from execution of the test (per cluster) */
-	TArray< FAutomationTestResults > Results;
+	TArray< TArray<FAutomationTestResults> > Results;
 
 	/** Structure holding the test info */
 	FAutomationTestInfo TestInfo;
diff --git a/Engine/Source/Developer/AutomationController/Private/AutomationReportManager.cpp b/Engine/Source/Developer/AutomationController/Private/AutomationReportManager.cpp
--- a/Engine/Source/Developer/AutomationController/Private/AutomationReportManager.cpp
+++ b/Engine/Source/Developer/AutomationController/Private/AutomationReportManager.cpp
@@ -22,10 +22,10 @@ void FAutomationReportManager::Empty()
 }
 
 
-void FAutomationReportManager::ResetForExecution()
+void FAutomationReportManager::ResetForExecution(const int32 NumTestPasses)
 {
 	//Recursively prepare all tests for execution
-	ReportRoot->ResetForExecution();
+	ReportRoot->ResetForExecution(NumTestPasses);
 }
 
 
@@ -34,16 +34,16 @@ void FAutomationReportManager::StopRunningTests()
 	ReportRoot->StopRunningTest();
 }
 
-TSharedPtr<IAutomationReport> FAutomationReportManager::GetNextReportToExecute(bool& bOutAllTestsComplete, const int32 ClusterIndex, const int32 NumDevicesInCluster)
+TSharedPtr<IAutomationReport> FAutomationReportManager::GetNextReportToExecute(bool& bOutAllTestsComplete, const int32 ClusterIndex, const int32 PassIndex, const int32 NumDevicesInCluster)
 {
-	TSharedPtr<IAutomationReport> ReportToExecute = ReportRoot->GetNextReportToExecute(bOutAllTestsComplete, ClusterIndex, NumDevicesInCluster);
+	TSharedPtr<IAutomationReport> ReportToExecute = ReportRoot->GetNextReportToExecute(bOutAllTestsComplete, ClusterIndex, PassIndex, NumDevicesInCluster);
 	return ReportToExecute;
 }
 
 
-TSharedPtr<IAutomationReport> FAutomationReportManager::EnsureReportExists(FAutomationTestInfo& TestInfo, const int32 ClusterIndex)
+TSharedPtr<IAutomationReport> FAutomationReportManager::EnsureReportExists(FAutomationTestInfo& TestInfo, const int32 ClusterIndex, const int32 NumPasses)
 {
-	TSharedPtr<IAutomationReport> NewStatus = ReportRoot->EnsureReportExists (TestInfo, ClusterIndex);
+	TSharedPtr<IAutomationReport> NewStatus = ReportRoot->EnsureReportExists (TestInfo, ClusterIndex, NumPasses);
 	return NewStatus;
 }
 
@@ -71,6 +71,15 @@ int32 FAutomationReportManager::GetEnabledTestsNum() const
 	return ReportRoot->GetEnabledTestsNum();
 }
 
+void FAutomationReportManager::GetEnabledTestNames(TArray<FString>& OutEnabledTestNames) const
+{
+	ReportRoot->GetEnabledTestNames(OutEnabledTestNames,FString());
+}
+
+void FAutomationReportManager::SetEnabledTests(const TArray<FString>& EnabledTests)
+{
+	ReportRoot->SetEnabledTests(EnabledTests,FString());
+}
 
 const bool FAutomationReportManager::ExportReport(uint32 FileExportTypeMask, const int32 NumDeviceClusters)
 {
@@ -121,7 +130,7 @@ void FAutomationReportManager::AddResultReport(TSharedPtr< IAutomationReport > I
 	{
 		for (int32 ClusterIndex = 0; ClusterIndex < NumDeviceClusters; ++ClusterIndex)
 		{
-			FAutomationTestResults TestResults = InReport->GetResults( ClusterIndex );
+			FAutomationTestResults TestResults = InReport->GetResults( ClusterIndex,CurrentTestPass );
 
 			// Early out if we don't want this report
 			bool AddReport = true;
@@ -172,11 +181,11 @@ void FAutomationReportManager::AddResultReport(TSharedPtr< IAutomationReport > I
 				FString Status;
 
 				// Was the test a success
-				if (( TestResults.Warnings.Num() == 0 ) && ( TestResults.Errors.Num() == 0 ) && ( InReport->GetState( ClusterIndex ) == EAutomationState::Success ))
+				if (( TestResults.Warnings.Num() == 0 ) && ( TestResults.Errors.Num() == 0 ) && ( InReport->GetState( ClusterIndex,CurrentTestPass ) == EAutomationState::Success ))
 				{
 					Status = "Success";
 				}
-				else if( InReport->GetState( ClusterIndex ) == EAutomationState::NotEnoughParticipants )
+				else if( InReport->GetState( ClusterIndex,CurrentTestPass ) == EAutomationState::NotEnoughParticipants )
 				{
 					Status = TEXT( "Could not run." );
 				}
diff --git a/Engine/Source/Developer/AutomationController/Private/AutomationReportManager.h b/Engine/Source/Developer/AutomationController/Private/AutomationReportManager.h
--- a/Engine/Source/Developer/AutomationController/Private/AutomationReportManager.h
+++ b/Engine/Source/Developer/AutomationController/Private/AutomationReportManager.h
@@ -19,8 +19,11 @@ public:
 	/** Clears out all reports */
 	void Empty();
 
-	/** Resets for the next run of tests */
-	void ResetForExecution();
+	/** 
+	 * Resets for the next run of tests 
+	 * @param NumTestPasses - The number of test results to store
+	 */
+	void ResetForExecution(const int32 NumTestPasses);
 
 	/** Stops all tests from running */
 	void StopRunningTests();
@@ -29,17 +32,19 @@ public:
 	 * Searches for the next test to execute 
 	 * @param bOutAllTestsComplete - Returns True if all tests for this device cluster have been completed
 	 * @param ClusterIndex - Index of the platform reporting the results of this test.  See AutomationDeviceClusterManager
+	 * @param PassIndex - Index of the current test pass we are on
 	 * @param NumDevicesInCluster - The number of devices which are available in this cluster
 	 */
-	TSharedPtr<IAutomationReport> GetNextReportToExecute(bool& bOutAllTestsComplete, const int32 ClusterIndex, const int32 NumDevicesInCluster);
+	TSharedPtr<IAutomationReport> GetNextReportToExecute(bool& bOutAllTestsComplete, const int32 ClusterIndex, const int32 PassIndex, const int32 NumDevicesInCluster);
 
 	/**
 	 * Ensures the nested tree exists ("blueprint.test.all") would have 3 levels other than the root
 	 * @param TestInfo - Structure containing the test info
 	 * @param ClusterIndex - Index of the platform reporting the results of this test.  See AutomationDeviceClusterManager
+	 * @param NumPasses - The number of test passes to create reports for
 	 * @return The automation report
 	 */
-	TSharedPtr<IAutomationReport> EnsureReportExists(FAutomationTestInfo& TestInfo, const int32 ClusterIndex);
+	TSharedPtr<IAutomationReport> EnsureReportExists(FAutomationTestInfo& TestInfo, const int32 ClusterIndex, const int32 NumPasses);
 
 	/** Filters the visible tests based on name, status, speed */
 	void SetFilter( TSharedPtr< AutomationFilterCollection > InFilter );
@@ -52,6 +57,14 @@ public:
 	/** Returns number of tests that will be run */
 	int32 GetEnabledTestsNum () const;
 
+	/** Gets the names of enabled tests */
+	void GetEnabledTestNames(TArray<FString>& OutEnabledTestNames) const;
+	/** Sets the enabled tests based off the passed list of enabled tests */
+	void SetEnabledTests(const TArray<FString>& EnabledTests);
+
+	/** Sets the current test pass */
+	void SetCurrentTestPass(int32 PassIndex){ CurrentTestPass = PassIndex; }
+
 	/**
 	 * Export the automation report
 	 * @param FileExportTypeMask - The file export mask - errors, warning etc
@@ -83,6 +96,8 @@ protected:
 private:
 	/** Root node of the hierarchy, just there to ensure little code duplication */
 	TSharedPtr <IAutomationReport> ReportRoot;
+	/** Which test pass we are currently on */
+	int32 CurrentTestPass;
 };
 
 
diff --git a/Engine/Source/Developer/AutomationController/Public/Interfaces/IAutomationControllerManager.h b/Engine/Source/Developer/AutomationController/Public/Interfaces/IAutomationControllerManager.h
--- a/Engine/Source/Developer/AutomationController/Public/Interfaces/IAutomationControllerManager.h
+++ b/Engine/Source/Developer/AutomationController/Public/Interfaces/IAutomationControllerManager.h
@@ -140,6 +140,18 @@ public:
 	/** Tick function that will execute enabled tests for different device clusters */
 	virtual void Tick() = 0;
 
+	/**
+	 * Sets the number of test passes
+	 *
+	 * @param InNumPasses - the number of test passes to perform
+	 */
+	virtual void SetNumPasses(const int32 InNumPasses) = 0;
+
+	/**
+	 * Get the number of test passes
+	 */
+	virtual int32 GetNumPasses() = 0;
+
 	/**
 	 * Filters the visible list of tests.
 	 */
@@ -183,6 +195,20 @@ public:
 	 */
 	virtual int32 GetEnabledTestsNum() const = 0;
 
+	/**
+	 * Gets the names of all the enabled tests
+	 *
+	 * @param OutEnabledTestNames The list to populate with enabled test names
+	 */
+	virtual void GetEnabledTestNames(TArray<FString>& OutEnabledTestNames) const = 0;
+
+	/**
+	 * Sets any tests that match a name in the enabled tests array
+	 *
+	 * @param EnabledTests An array of test names that will be enabled
+	 */
+	virtual void SetEnabledTests(const TArray<FString>& EnabledTests) = 0;
+
 	/**
 	 * Gets the controller's current test state.
 	 *
diff --git a/Engine/Source/Developer/AutomationController/Public/Interfaces/IAutomationReport.h b/Engine/Source/Developer/AutomationController/Public/Interfaces/IAutomationReport.h
--- a/Engine/Source/Developer/AutomationController/Public/Interfaces/IAutomationReport.h
+++ b/Engine/Source/Developer/AutomationController/Public/Interfaces/IAutomationReport.h
@@ -147,6 +147,12 @@ public:
 	/** Recursively gets the number of child nodes */
 	virtual int32 GetTotalNumChildren() const = 0;
 
+	/** Gets the names of all the enabled tests */
+	virtual void GetEnabledTestNames(TArray<FString>& OutEnabledTestNames, FString CurrentPath) const = 0;
+
+	/** Sets which tests are enabled based off the enabled tests list */
+	virtual void SetEnabledTests(const TArray<FString>& EnabledTests, FString CurrentPath) = 0;
+
 	/** Recursively gets the number of enabled tests */
 	virtual int32 GetEnabledTestsNum() const = 0;
 
@@ -197,38 +203,59 @@ public:
 	/** Returns the array of child reports */
 	virtual TArray<TSharedPtr<IAutomationReport> >& GetChildReports() = 0;
 
-	/** Recursively resets the report to "needs to be run", clears cached warnings and errors */
-	virtual void ResetForExecution() = 0;
+	/** 
+	 * Recursively resets the report to "needs to be run", clears cached warnings and errors
+	 * @param NumTestPasses - The number of test passes so we know how many results to create
+	 */
+	virtual void ResetForExecution(const int32 NumTestPasses) = 0;
 
 	/**
 	 * Sets the results of the test for use by the UI
 	 * @param ClusterIndex - Index of the platform reporting the results of this test.  See AutomationDeviceClusterManager
+	 * @param PassIndex - Which test pass these results are for
 	 * @param InResults - The new set of results
 	 */
-	virtual void SetResults( const int32 ClusterIndex, const FAutomationTestResults& InResults ) = 0;
+	virtual void SetResults( const int32 ClusterIndex, const int32 PassIndex, const FAutomationTestResults& InResults ) = 0;
 
 	/**
 	 * Returns completion statistics for this branch of the testing hierarchy
 	 * @param ClusterIndex - Index of the platform reporting the results of this test.  See AutomationDeviceClusterManager
+	 * @param PassIndex - Which test pass to get the status of
 	 * @param OutCompletionState - Collection structure for execution statistics
 	 */
-	virtual void GetCompletionStatus(const int32 ClusterIndex, FAutomationCompleteState& OutCompletionState) = 0;
+	virtual void GetCompletionStatus(const int32 ClusterIndex, const int32 PassIndex, FAutomationCompleteState& OutCompletionState) = 0;
 
 	/**
 	 * Returns the state of the test (not run, in process, success, failure)
 	 * @param ClusterIndex - Index of the platform reporting the results of this test.  See AutomationDeviceClusterManager
+	 * @param PassIndex - Which test pass to get the state of
 	 * @return the current state of the test
 	 */
-	virtual EAutomationState::Type GetState(const int32 ClusterIndex) const = 0;
+	virtual EAutomationState::Type GetState(const int32 ClusterIndex, const int32 PassIndex) const = 0;
 
 	/**
 	 * Gets a copy of errors and warnings that were found
 	 *
 	 * @param ClusterIndex - Index of the platform we are requesting test results for.
+	 * @param PassIndex - Index of the test pass to get the results for.
 	 *
 	 * @return The collection of results for the given cluster index
 	 */
-	virtual const FAutomationTestResults& GetResults( const int32 ClusterIndex ) = 0;
+	virtual const FAutomationTestResults& GetResults( const int32 ClusterIndex, const int32 PassIndex ) = 0;
+
+	/**
+	 * Gets the number of available test results for a given cluster
+	 * @param ClusterIndex - Index of the platform .
+	 * @return The number of results available for this cluster
+	 */
+	virtual const int32 GetNumResults( const int32 ClusterIndex ) = 0;
+
+	/**
+	 * Finds the current pass by looking at the current state
+	 * @param ClusterIndex - Index of the platform.
+	 * @return The current pass index
+	 */
+	virtual const int32 GetCurrentPassIndex( const int32 ClusterIndex ) = 0;
 
 	/**
 	 * Gets the name of the instance that ran the test
@@ -242,19 +269,21 @@ public:
 	 * If NewTestName is Editor.Maps.Loadall.TestName, this will create nodes for Editor, Maps, Loadall, and then a leaf node for the testname with the associated commandline
 	 * @param TestInfo - Structure containing all the test info
 	 * @param ClusterIndex - Index of the platform reporting the results of this test.  See AutomationDeviceClusterManager
+	 * @param NumPasses - The number of passes we are going to perform.  Used to make sure we have enough results.
 	 * @return - the automation report
 	 */
-	virtual TSharedPtr<IAutomationReport> EnsureReportExists(FAutomationTestInfo& TestInfo, const int32 ClusterIndex) = 0;
+	virtual TSharedPtr<IAutomationReport> EnsureReportExists(FAutomationTestInfo& TestInfo, const int32 ClusterIndex, const int32 NumPasses) = 0;
 
 	/**
 	 * Returns the next test in the hierarchy to run
 	 * @param bOutAllTestsComplete - Whether or not all enabled tests have completed execution for this platform (cluster index)
 	 * @param ClusterIndex - Index of the platform reporting the results of this test.  See AutomationDeviceClusterManager
+	 * @param PassIndex - Which test pass we are currently on
 	 * @param NumDevicesInCluster - The number of devices which are in this cluster
 	 *
 	 * @return The next report
 	 */
-	virtual TSharedPtr<IAutomationReport> GetNextReportToExecute(bool& bOutAllTestsComplete, const int32 ClusterIndex, const int32 NumDevicesInCluster) = 0;
+	virtual TSharedPtr<IAutomationReport> GetNextReportToExecute(bool& bOutAllTestsComplete, const int32 ClusterIndex, const int32 PassIndex, const int32 NumDevicesInCluster) = 0;
 
 	/**
 	 * Returns if there have been any errors in the test
diff --git a/Engine/Source/Developer/AutomationWindow/Private/AutomationPresetManager.cpp b/Engine/Source/Developer/AutomationWindow/Private/AutomationPresetManager.cpp
new file mode 100644
--- /dev/null
+++ b/Engine/Source/Developer/AutomationWindow/Private/AutomationPresetManager.cpp
@@ -0,0 +1,101 @@
+// Copyright 1998-2014 Epic Games, Inc. All Rights Reserved.
+
+#include "AutomationWindowPrivatePCH.h"
+
+FAutomationTestPresetManager::FAutomationTestPresetManager()
+{
+}
+
+AutomationPresetRef FAutomationTestPresetManager::AddNewPreset( )
+{
+	AutomationPresetPtr NewPreset = MakeShareable(new FAutomationTestPreset());
+	Presets.Add(NewPreset);
+	return NewPreset.ToSharedRef();
+}
+
+AutomationPresetRef FAutomationTestPresetManager::AddNewPreset( const FString& PresetName, const TArray<FString>& SelectedTests )
+{
+	AutomationPresetRef NewPreset = AddNewPreset();
+	NewPreset->SetPresetName(PresetName);
+	NewPreset->SetEnabledTests(SelectedTests);
+	SavePreset(NewPreset);
+	return NewPreset;
+}
+
+TArray<AutomationPresetPtr>& FAutomationTestPresetManager::GetAllPresets( )
+{
+	return Presets;
+}
+
+AutomationPresetPtr FAutomationTestPresetManager::LoadPreset( FArchive& Archive )
+{
+	FAutomationTestPreset* NewPreset = new FAutomationTestPreset();
+
+	if(NewPreset->Serialize(Archive))
+	{
+		return MakeShareable(NewPreset);
+	}
+
+	delete NewPreset;
+
+	return NULL;
+}
+
+void FAutomationTestPresetManager::RemovePreset( const AutomationPresetRef Preset )
+{
+	if (Presets.Remove(Preset) > 0)
+	{
+		// Find the name of the preset on disk
+		FString PresetFileName = GetPresetFolder() / Preset->GetID().ToString() + TEXT(".uap");
+
+		// delete the preset on disk
+		IFileManager::Get().Delete(*PresetFileName);
+	}
+}
+
+void FAutomationTestPresetManager::SavePreset( const AutomationPresetRef Preset )
+{
+	FString PresetFileName = GetPresetFolder() / Preset->GetID().ToString() + TEXT(".uap");
+	FArchive* PresetFileWriter = IFileManager::Get().CreateFileWriter(*PresetFileName);
+
+	if (PresetFileWriter != NULL)
+	{
+		SavePreset(Preset, *PresetFileWriter);
+
+		delete PresetFileWriter;
+	}
+}
+
+void FAutomationTestPresetManager::SavePreset( const AutomationPresetRef Preset, FArchive& Archive )
+{
+	Preset->Serialize(Archive);
+}
+
+void FAutomationTestPresetManager::LoadPresets( )
+{
+	TArray<FString> PresetFileNames;
+
+	IFileManager::Get().FindFiles(PresetFileNames, *(GetPresetFolder() / TEXT("*.uap")), true, false);
+
+	for (TArray<FString>::TConstIterator It(PresetFileNames); It; ++It)
+	{
+		FString PresetFilePath = GetPresetFolder() / *It;
+		FArchive* PresetFileReader = IFileManager::Get().CreateFileReader(*PresetFilePath);
+
+		if (PresetFileReader != NULL)
+		{
+			AutomationPresetPtr LoadedPreset = LoadPreset(*PresetFileReader);
+
+			if (LoadedPreset.IsValid())
+			{
+				Presets.Add(LoadedPreset);
+			}
+			else
+			{
+				IFileManager::Get().Delete(*PresetFilePath);
+			}
+
+			delete PresetFileReader;
+		}
+	}
+}
diff --git a/Engine/Source/Developer/AutomationWindow/Private/AutomationPresetManager.h b/Engine/Source/Developer/AutomationWindow/Private/AutomationPresetManager.h
new file mode 100644
--- /dev/null
+++ b/Engine/Source/Developer/AutomationWindow/Private/AutomationPresetManager.h
@@ -0,0 +1,181 @@
+// Copyright 1998-2014 Epic Games, Inc. All Rights Reserved.
+
+#pragma once
+
+/**
+ * Type definition for shared pointers to instances of FAutomationTestPreset.
+ */
+typedef TSharedPtr<class FAutomationTestPreset> AutomationPresetPtr;
+
+/**
+ * Type definition for shared references to instances of FAutomationTestPreset.
+ */
+typedef TSharedRef<class FAutomationTestPreset> AutomationPresetRef;
+
+//Increment this number if the Preset serialization changes.
+#define AUTOMATION_PRESETVERSION 1
+
+/**
+ * Class that holds preset data for the automation window
+ */
+class FAutomationTestPreset
+{
+public:
+
+	/**
+	 * Constructor.
+	 */
+	FAutomationTestPreset() :
+		Id(FGuid::NewGuid())
+	{
+	}
+
+	/**
+	 * Destructor.
+	 */
+	~FAutomationTestPreset( ) { }
+
+	/**
+	 * Gets the GUid for this preset
+	 */
+	const FGuid GetID()
+	{
+		return Id;
+	}
+
+	/**
+	 * Gets the name for this preset
+	 */
+	const FString& GetPresetName( ) const
+	{
+		return PresetName;
+	}
+
+	/**
+	 * Sets the name for this preset
+	 */
+	void SetPresetName(const FString& InPresetName)
+	{
+		PresetName = InPresetName;
+	}
+
+	/**
+	 * Returns the list of enabled tests
+	 */
+	const TArray<FString>& GetEnabledTests() const
+	{
+		return EnabledTests;
+	}
+
+	/**
+	 * Sets the list of enabled tests
+	 */
+	void SetEnabledTests(TArray<FString> NewEnabledTests)
+	{
+		EnabledTests = NewEnabledTests;
+	}
+
+	/**
+	 * Handles saving / loading the preset data to an archive (file)
+	 */
+	virtual bool Serialize( FArchive& Archive )
+	{
+		int32 Version = AUTOMATION_PRESETVERSION;
+
+		Archive	<< Version;
+
+		if (Version != AUTOMATION_PRESETVERSION)
+		{
+			return false;
+		}
+
+		Archive << Id;
+		Archive << PresetName;
+		Archive << EnabledTests;
+
+		return true;
+	}
+
+private:
+	// A unique ID for this preset. (Used for the filename because the PresetName may have invalid characters)
+	FGuid Id;
+
+	// The name of this preset
+	FString PresetName;
+
+	// The list of enabled test names
+	TArray<FString> EnabledTests;
+};
+
+class FAutomationTestPresetManager
+{
+public:
+	/**
+	 * Constructor.
+	 */
+	FAutomationTestPresetManager();
+
+	/**
+	 * Creates a new empty preset
+	 */
+	virtual AutomationPresetRef AddNewPreset( );
+
+	/**
+	 * Creates a new preset with the given name and enabled tests
+	 * @param PresetName - The name of the new preset
+	 * @param SelectedTests - The list of enabled tests
+	 */
+	virtual AutomationPresetRef AddNewPreset( const FString& PresetName, const TArray<FString>& SelectedTests );
+
+	/**
+	 * Returns a reference to the list that holds the presets
+	 */
+	virtual TArray<AutomationPresetPtr>& GetAllPresets( );
+
+	/**
+	 * Removes the selected preset from the preset list and deletes the file on disk
+	 * @param Preset - The preset to remove
+	 */
+	void RemovePreset( const AutomationPresetRef Preset );
+
+	/**
+	 * Saves the passed in preset to disk
+	 * @param Preset - The preset to save
+	 */
+	void SavePreset( const AutomationPresetRef Preset);
+
+	/*
+	* Load all Presets from disk.
+	*/
+	void LoadPresets( );
+
+protected:
+
+	/**
+	 * Creates a new preset and loads it with data from the archive
+	 * @param Archive - The stream to read the preset data from
+	 */
+	AutomationPresetPtr LoadPreset( FArchive& Archive );
+
+	/**
+	 * Writes the Preset data to the archive
+	 * @param Preset - The preset to serialize
+	 * @param Archive - The stream to write the preset data to
+	 */
+	void SavePreset( const AutomationPresetRef Preset, FArchive& Archive );
+
+	/**
+	 * Gets the folder in which preset files are stored.
+	 *
+	 * @return The folder path.
+	 */
+	static FString GetPresetFolder( )
+	{
+		return FPaths::EngineSavedDir() / TEXT("Automation");
+	}
+
+private:
+
+	// Holds the collection of automation presets.
+	TArray<AutomationPresetPtr> Presets;
+};
diff --git a/Engine/Source/Developer/AutomationWindow/Private/AutomationWindowPrivatePCH.h b/Engine/Source/Developer/AutomationWindow/Private/AutomationWindowPrivatePCH.h
--- a/Engine/Source/Developer/AutomationWindow/Private/AutomationWindowPrivatePCH.h
+++ b/Engine/Source/Developer/AutomationWindow/Private/AutomationWindowPrivatePCH.h
@@ -26,6 +26,7 @@
 #include "SAutomationWindowCommandBar.h"
 #include "SAutomationExportMenu.h"
 #include "AutomationFilter.h"
+#include "AutomationPresetManager.h"
 #include "SAutomationWindow.h"
 #include "SAutomationTestItemContextMenu.h"
 #include "SAutomationTestItem.h"
diff --git a/Engine/Source/Developer/AutomationWindow/Private/SAutomationTestItem.cpp b/Engine/Source/Developer/AutomationWindow/Private/SAutomationTestItem.cpp
--- a/Engine/Source/Developer/AutomationWindow/Private/SAutomationTestItem.cpp
+++ b/Engine/Source/Developer/AutomationWindow/Private/SAutomationTestItem.cpp
@@ -215,7 +215,8 @@ const FSlateBrush* SAutomationTestItem::GetSmokeTestImage() const
 FText SAutomationTestItem::GetTestToolTip( int32 ClusterIndex ) const
 {
 	FText ToolTip;
-	EAutomationState::Type TestState = TestStatus->GetState( ClusterIndex );
+	const int32 PassIndex = TestStatus->GetCurrentPassIndex(ClusterIndex);
+	EAutomationState::Type TestState = TestStatus->GetState( ClusterIndex, PassIndex );
 	if ( TestState == EAutomationState::NotRun )
 	{
 		ToolTip = LOCTEXT("TestToolTipNotRun", "Not Run");
@@ -256,7 +257,8 @@ FSlateColor SAutomationTestItem::ItemStatus_BackgroundColor(const int32 ClusterI
 {
 	if (TestStatus->GetTotalNumChildren()==0)
 	{
-		EAutomationState::Type TestState = TestStatus->GetState(ClusterIndex);
+		const int32 PassIndex = TestStatus->GetCurrentPassIndex(ClusterIndex);
+		EAutomationState::Type TestState = TestStatus->GetState(ClusterIndex,PassIndex);
 		if (TestState == EAutomationState::Fail)
 		{
 			// Failure is marked by a red background.
@@ -315,7 +317,8 @@ FText SAutomationTestItem::ItemStatus_DurationText() const
 
 EVisibility SAutomationTestItem::ItemStatus_GetStatusVisibility(const int32 ClusterIndex, const bool bForInProcessThrobber) const
 {
-	EAutomationState::Type TestState = TestStatus->GetState(ClusterIndex);
+	const int32 PassIndex = TestStatus->GetCurrentPassIndex(ClusterIndex);
+	EAutomationState::Type TestState = TestStatus->GetState(ClusterIndex,PassIndex);
 	bool bImageVisible = TestState != EAutomationState::InProcess;
 
 	bool bFinalVisibility =  bForInProcessThrobber ? !bImageVisible : bImageVisible;
@@ -336,7 +339,8 @@ FText SAutomationTestItem::ItemStatus_NumParticipantsRequiredText() const
 FSlateColor SAutomationTestItem::ItemStatus_ProgressColor(const int32 ClusterIndex) const
 {
 	FAutomationCompleteState CompleteState;
-	TestStatus->GetCompletionStatus(ClusterIndex, CompleteState);
+	const int32 PassIndex = TestStatus->GetCurrentPassIndex(ClusterIndex);
+	TestStatus->GetCompletionStatus(ClusterIndex,PassIndex, CompleteState);
 
 	if (CompleteState.TotalEnabled > 0)
 	{
@@ -367,7 +371,8 @@ FSlateColor SAutomationTestItem::ItemStatus_ProgressColor(const int32 ClusterInd
 TOptional<float> SAutomationTestItem::ItemStatus_ProgressFraction(const int32 ClusterIndex) const
 {
 	FAutomationCompleteState CompleteState;
-	TestStatus->GetCompletionStatus(ClusterIndex, CompleteState);
+	const int32 PassIndex = TestStatus->GetCurrentPassIndex(ClusterIndex);
+	TestStatus->GetCompletionStatus(ClusterIndex,PassIndex, CompleteState);
 
 	uint32 TotalComplete = CompleteState.NumEnabledTestsPassed + CompleteState.NumEnabledTestsFailed + CompleteState.NumEnabledTestsCouldntBeRun;
 	// Only show a percentage if there is something interesting to report
@@ -382,7 +387,8 @@ TOptional<float> SAutomationTestItem::ItemStatus_ProgressFraction(const int32 Cl
 
 const FSlateBrush* SAutomationTestItem::ItemStatus_StatusImage(const int32 ClusterIndex) const
 {
-	EAutomationState::Type TestState = TestStatus->GetState(ClusterIndex);
+	const int32 PassIndex = TestStatus->GetCurrentPassIndex(ClusterIndex);
+	EAutomationState::Type TestState = TestStatus->GetState(ClusterIndex,PassIndex);
 
 	const FSlateBrush* ImageToUse;
 	switch( TestState )
@@ -390,7 +396,7 @@ const FSlateBrush* SAutomationTestItem::ItemStatus_StatusImage(const int32 Clust
 	case EAutomationState::Success:
 		{
 			FAutomationCompleteState CompleteState;
-			TestStatus->GetCompletionStatus(ClusterIndex, CompleteState);
+			TestStatus->GetCompletionStatus(ClusterIndex,PassIndex, CompleteState);
 			//If there were ANY warnings in the results
 			if (CompleteState.NumEnabledTestsWarnings || CompleteState.NumDisabledTestsWarnings)
 			{
diff --git a/Engine/Source/Developer/AutomationWindow/Private/SAutomationWindow.cpp b/Engine/Source/Developer/AutomationWindow/Private/SAutomationWindow.cpp
--- a/Engine/Source/Developer/AutomationWindow/Private/SAutomationWindow.cpp
+++ b/Engine/Source/Developer/AutomationWindow/Private/SAutomationWindow.cpp
@@ -3,6 +3,7 @@
 #include "AutomationWindowPrivatePCH.h"
 
 #include "AutomationController.h"
+#include "AutomationPresetManager.h"
 #include "SAutomationWindow.h"
 
 #define LOCTEXT_NAMESPACE "AutomationTest"
@@ -62,6 +63,7 @@ SAutomationWindow::~SAutomationWindow()
 	{
 		SessionManager->OnCanSelectSession().RemoveAll(this);
 		SessionManager->OnSelectedSessionChanged().RemoveAll(this);
+		SessionManager->OnSessionInstanceUpdated().RemoveAll(this);
 	}
 
 	if (AutomationController.IsValid())
@@ -76,6 +78,10 @@ void SAutomationWindow::Construct( const FArguments& InArgs, const IAutomationCo
 	FAutomationWindowCommands::Register();
 	CreateCommands();
 
+	TestPresetManager = MakeShareable(new FAutomationTestPresetManager());
+	TestPresetManager->LoadPresets();
+	bAddingTestPreset = false;
+
 	SessionManager = InSessionManager;
 	AutomationController = InAutomationController;
 
@@ -210,8 +216,8 @@ void SAutomationWindow::Construct( const FArguments& InArgs, const IAutomationCo
 					+SOverlay::Slot()
 					[
 						SNew(SBorder)
-						.BorderImage(FEditorStyle::GetBrush("ToolPanel.GroupBorder"))
-						.Padding(0.0f)
+						.BorderImage(this, &SAutomationWindow::GetTestBackgroundBorderImage)
+						.Padding(3.0f)
 						[
 							//the actual table full of tests
 							TestTable.ToSharedRef()
@@ -292,11 +298,28 @@ void SAutomationWindow::Construct( const FArguments& InArgs, const IAutomationCo
 
 	SessionManager->OnCanSelectSession().AddSP( this, &SAutomationWindow::HandleSessionManagerCanSelectSession );
 	SessionManager->OnSelectedSessionChanged().AddSP( this, &SAutomationWindow::HandleSessionManagerSelectionChanged );
+	SessionManager->OnSessionInstanceUpdated().AddSP( this, &SAutomationWindow::HandleSessionManagerInstanceChanged );
 
 	FindWorkers();
 }
 END_SLATE_FUNCTION_BUILD_OPTIMIZATION
 
+const FSlateBrush* SAutomationWindow::GetTestBackgroundBorderImage() const
+{
+	switch(TestBackgroundType)
+	{
+	case EAutomationTestBackgroundStyle::Game:
+		return FEditorStyle::GetBrush("AutomationWindow.GameGroupBorder");
+
+	case EAutomationTestBackgroundStyle::Editor:
+		return FEditorStyle::GetBrush("AutomationWindow.EditorGroupBorder");
+
+	case EAutomationTestBackgroundStyle::Unknown:
+	default:
+		return FEditorStyle::GetBrush("ToolPanel.GroupBorder");
+	}
+}
+
 void SAutomationWindow::CreateCommands()
 {
 	check(!AutomationWindowActions.IsValid());
@@ -356,11 +379,19 @@ TSharedRef< SWidget > SAutomationWindow::MakeAutomationWindowToolBar( const TSha
 {
 	struct Local
 	{
-		static void FillToolbar(FToolBarBuilder& ToolbarBuilder, TSharedRef<SWidget> RunTests, TSharedRef<SWidget> Searchbox)
+		static void FillToolbar(FToolBarBuilder& ToolbarBuilder, TSharedRef<SWidget> RunTests, TSharedRef<SWidget> Searchbox, TSharedRef<SWidget> PresetBox, TWeakPtr<class SAutomationWindow> InAutomationWindow )
 		{
 			ToolbarBuilder.BeginSection("Automation");
 			{
 				ToolbarBuilder.AddWidget( RunTests );
+				FUIAction DefaultAction;
+				ToolbarBuilder.AddComboButton(
+					DefaultAction,
+					FOnGetContent::CreateStatic( &SAutomationWindow::GenerateTestsOptionsMenuContent, InAutomationWindow ),
+					LOCTEXT( "TestOptions_Label", "Test Options" ),
+					LOCTEXT( "TestOptionsToolTip", "Test Options" ),
+					FSlateIcon(FEditorStyle::GetStyleSetName(), "AutomationWindow.TestOptions"),
+					true);
 				ToolbarBuilder.AddToolBarButton( FAutomationWindowCommands::Get().RefreshTests );
 				ToolbarBuilder.AddToolBarButton( FAutomationWindowCommands::Get().FindWorkers );
 			}
@@ -379,6 +410,11 @@ TSharedRef< SWidget > SAutomationWindow::MakeAutomationWindowToolBar( const TSha
 				ToolbarBuilder.AddWidget( Searchbox );
 			}
 			ToolbarBuilder.EndSection();
+			ToolbarBuilder.BeginSection("Presets");
+			{
+				ToolbarBuilder.AddWidget( PresetBox );
+			}
+			ToolbarBuilder.EndSection();
 		}
 	};
 
@@ -446,8 +482,108 @@ TSharedRef< SWidget > SAutomationWindow::MakeAutomationWindowToolBar( const TSha
 			.MinDesiredWidth(SearchWidth)
 		];
 
+	TSharedRef<SWidget> TestPresets = 
+		SNew( SHorizontalBox )
+		+SHorizontalBox::Slot()
+		.AutoWidth()
+		.Padding(2.0f)
+		.VAlign(VAlign_Bottom)
+		[
+			SNew( STextBlock )
+			.Text( LOCTEXT("AutomationPresetLabel", "Preset:") )
+		]
+		+SHorizontalBox::Slot()
+		.AutoWidth()
+		.Padding(2.0f)
+		[
+			//Preset Combo / Text
+			SNew(SOverlay)
+			+SOverlay::Slot()
+			[
+				SNew(SHorizontalBox)
+				.Visibility(this,&SAutomationWindow::HandlePresetComboVisibility)
+				+SHorizontalBox::Slot()
+				.FillWidth(1.0)
+				[
+					SAssignNew( PresetComboBox, SComboBox< TSharedPtr<FAutomationTestPreset> > )
+					.OptionsSource( &TestPresetManager->GetAllPresets() )
+					.OnGenerateWidget( this, &SAutomationWindow::GeneratePresetComboItem )
+					.OnSelectionChanged( this, &SAutomationWindow::HandlePresetChanged )
+					[
+						SNew( STextBlock )
+						.Text( this, &SAutomationWindow::GetPresetComboText )
+					]
+				]
+			]
+			+SOverlay::Slot()
+			[
+				SNew(SHorizontalBox)
+				.Visibility(this,&SAutomationWindow::HandlePresetTextVisibility)
+				+SHorizontalBox::Slot()
+				.FillWidth(1.0)
+				[
+					SAssignNew(PresetTextBox, SEditableTextBox)
+					.OnTextCommitted(this, &SAutomationWindow::HandlePresetTextCommited)
+				]
+			]
+		]
+
+		//New button
+		+SHorizontalBox::Slot()
+		.AutoWidth()
+		.Padding(2.0f)
+		[
+			SNew( SButton )
+			.ButtonStyle( FEditorStyle::Get(), "NoBorder" )
+			.OnClicked( this, &SAutomationWindow::HandleNewPresetClicked )
+			.ToolTipText( LOCTEXT("AutomationPresetNewButtonTooltip", "Create a new preset") )
+			.IsEnabled(this, &SAutomationWindow::IsAddButtonIsEnabled)
+			.Content()
+			[
+				SNew(SImage)
+				.Image(FEditorStyle::Get().GetBrush("AutomationWindow.PresetNew"))
+			]
+		]
+
+		//Save button
+		+SHorizontalBox::Slot()
+		.AutoWidth()
+		.Padding(2.0f)
+		[
+			SNew( SButton )
+			.ButtonStyle( FEditorStyle::Get(), "NoBorder" )
+			.OnClicked( this, &SAutomationWindow::HandleSavePresetClicked )
+			.ToolTipText( LOCTEXT("AutomationPresetSaveButtonTooltip", "Save the current test list") )
+			.IsEnabled(this, &SAutomationWindow::IsSaveButtonIsEnabled)
+			.Content()
+			[
+				SNew(SImage)
+				.Image(FEditorStyle::Get().GetBrush("AutomationWindow.PresetSave"))
+			]
+		]
+
+		//Remove button
+		+ SHorizontalBox::Slot()
+		.AutoWidth()
+		.Padding(2.0f)
+		[
+			// remove button
+			SNew(SButton)
+			.ContentPadding(2)
+			.ButtonStyle( FEditorStyle::Get(), "NoBorder" )
+			.OnClicked( this, &SAutomationWindow::HandleRemovePresetClicked )
+			.ToolTipText(LOCTEXT("AutomationPresetRemoveButtonTooltip", "Remove the selected preset"))
+			.IsEnabled(this, &SAutomationWindow::IsRemoveButtonIsEnabled)
+			.Content()
+			[
+				SNew(SImage)
+				.Image(FEditorStyle::Get().GetBrush("AutomationWindow.PresetRemove"))
+			]
+		];
+
 	FToolBarBuilder ToolbarBuilder( InCommandList, FMultiBoxCustomization::None );
-	Local::FillToolbar( ToolbarBuilder, RunTests, Searchbox );
+	TWeakPtr<SAutomationWindow> AutomationWindow = SharedThis(this);
+	Local::FillToolbar( ToolbarBuilder, RunTests, Searchbox, TestPresets, AutomationWindow );
 
 	// Create the tool bar!
 	return
@@ -466,6 +602,158 @@ TSharedRef< SWidget > SAutomationWindow::MakeAutomationWindowToolBar( const TSha
 }
 END_SLATE_FUNCTION_BUILD_OPTIMIZATION
 
+EVisibility SAutomationWindow::HandlePresetComboVisibility( ) const
+{
+	return bAddingTestPreset ? EVisibility::Hidden : EVisibility::Visible;
+}
+
+EVisibility SAutomationWindow::HandlePresetTextVisibility( ) const
+{
+	return bAddingTestPreset ? EVisibility::Visible : EVisibility::Hidden;
+}
+
+bool SAutomationWindow::IsAddButtonIsEnabled() const
+{
+	return !bAddingTestPreset;
+}
+
+bool SAutomationWindow::IsSaveButtonIsEnabled() const
+{
+	return (!bAddingTestPreset && SelectedPreset.IsValid());
+}
+
+bool SAutomationWindow::IsRemoveButtonIsEnabled() const
+{
+	return (!bAddingTestPreset && SelectedPreset.IsValid());
+}
+
+void SAutomationWindow::HandlePresetTextCommited( const FText& CommittedText, ETextCommit::Type CommitType )
+{
+	if( CommitType == ETextCommit::OnEnter )
+	{
+		bAddingTestPreset = false;
+		if(CommittedText.IsEmpty())
+		{
+			return;
+		}
+
+		TArray<FString> EnabledTests;
+		AutomationController->GetEnabledTestNames(EnabledTests);
+		AutomationPresetPtr NewPreset = TestPresetManager->AddNewPreset(CommittedText.ToString(),EnabledTests);
+		PresetComboBox->SetSelectedItem(NewPreset);
+		SelectedPreset = NewPreset;
+
+		PresetTextBox->SetText(FText());
+	}
+	else if( CommitType == ETextCommit::OnCleared || CommitType == ETextCommit::OnUserMovedFocus )
+	{
+		if( bAddingTestPreset )
+		{
+			bAddingTestPreset = false;
+			SelectedPreset = NULL;
+			PresetComboBox->ClearSelection();
+			PresetTextBox->SetText(FText());
+		}
+	}
+}
+
+void SAutomationWindow::HandlePresetChanged( TSharedPtr<FAutomationTestPreset> Item, ESelectInfo::Type SelectInfo )
+{
+	if( Item.IsValid() )
+	{
+		SelectedPreset = Item;
+		AutomationController->SetEnabledTests(Item->GetEnabledTests());
+		TestTable->RequestTreeRefresh();
+	}
+}
+
+FReply SAutomationWindow::HandleNewPresetClicked()
+{
+	bAddingTestPreset = true;
+	return FReply::Handled().SetKeyboardFocus(PresetTextBox.ToSharedRef(), EKeyboardFocusCause::SetDirectly);
+}
+
+FReply SAutomationWindow::HandleSavePresetClicked()
+{
+	if(SelectedPreset.IsValid())
+	{
+		TArray<FString> EnabledTests;
+		AutomationController->GetEnabledTestNames(EnabledTests);
+		SelectedPreset->SetEnabledTests(EnabledTests);
+		TestPresetManager->SavePreset(SelectedPreset.ToSharedRef());
+	}
+	return FReply::Handled();
+}
+
+FReply SAutomationWindow::HandleRemovePresetClicked()
+{
+	if(SelectedPreset.IsValid())
+	{
+		TestPresetManager->RemovePreset(SelectedPreset.ToSharedRef());
+		SelectedPreset = NULL;
+		PresetComboBox->ClearSelection();
+	}
+	return FReply::Handled();
+}
+
+FString SAutomationWindow::GetPresetComboText() const
+{
+	if( SelectedPreset.IsValid() )
+	{
+		return *SelectedPreset->GetPresetName();
+	}
+	else
+	{
+		return LOCTEXT("AutomationPresetLabel", "<Select A Preset>").ToString();
+	}
+}
+
+TSharedRef<SWidget> SAutomationWindow::GeneratePresetComboItem(TSharedPtr<FAutomationTestPreset> InItem)
+{
+	return SNew(STextBlock) 
+		.Text( InItem->GetPresetName() );
+}
+
+TSharedRef< SWidget > SAutomationWindow::GenerateTestsOptionsMenuContent( TWeakPtr<class SAutomationWindow> InAutomationWindow )
+{
+	TSharedPtr<SAutomationWindow> AutomationWindow(InAutomationWindow.Pin());
+	if( AutomationWindow.IsValid() )
+	{
+		return AutomationWindow->GenerateTestsOptionsMenuContent();
+	}
+
+	//Return empty menu
+	FMenuBuilder MenuBuilder( true, NULL );
+	MenuBuilder.BeginSection("AutomationWindowRunTest", LOCTEXT("RunTestOptions", "Advanced Settings"));
+	MenuBuilder.EndSection();
+	return MenuBuilder.MakeWidget();
+}
+
+TSharedRef< SWidget > SAutomationWindow::GenerateTestsOptionsMenuContent( )
+{
+	const bool bShouldCloseWindowAfterMenuSelection = true;
+	FMenuBuilder MenuBuilder( bShouldCloseWindowAfterMenuSelection, AutomationWindowActions );
+	TSharedRef<SWidget> NumTests = SNew(SSpinBox<int32>)
+			.MinValue(1)
+			.MaxValue(10)
+			.MinSliderValue(1)
+			.MaxSliderValue(10)
+			.Value(this,&SAutomationWindow::GetRepeatCount)
+			.OnValueChanged(this,&SAutomationWindow::OnChangeRepeatCount);
+
+	MenuBuilder.BeginSection("AutomationWindowRunTest", LOCTEXT("RunTestOptions", "Advanced Settings"));
+	{
+		MenuBuilder.AddWidget( 
+			SNew( STextBlock )
+				.ColorAndOpacity(FSlateColor::UseSubduedForeground())
+				.Text(LOCTEXT("NumTestsToolTip", "Number of runs:")),
+			FText::GetEmpty() );
+		MenuBuilder.AddWidget(NumTests, FText::GetEmpty());
+	}
+
+	return MenuBuilder.MakeWidget();
+}
+
 // Only valid in the editor
 #if WITH_EDITOR
 TSharedPtr<SWidget> SAutomationWindow::HandleAutomationListContextMenuOpening()
@@ -513,10 +801,6 @@ void SAutomationWindow::OnTestSelectionChanged(TSharedPtr<IAutomationReport> Sel
 		int32 NumClusters = AutomationController->GetNumDeviceClusters();
 		for (int32 ClusterIndex = 0; ClusterIndex < NumClusters; ++ClusterIndex)
 		{
-			//get strings out of the report and populate the Log Messages
-
-			FAutomationTestResults TestResults = Selection->GetResults(ClusterIndex);
-
 			//no sense displaying device name if only one is available
 			if (NumClusters > 1)
 			{
@@ -524,24 +808,39 @@ void SAutomationWindow::OnTestSelectionChanged(TSharedPtr<IAutomationReport> Sel
 				LogMessages.Add(MakeShareable(new FAutomationOutputMessage(DeviceTypeName, TEXT("Automation.Header"))));
 			}
 
-			for (int32 ErrorIndex = 0; ErrorIndex < TestResults.Errors.Num(); ++ErrorIndex)
-			{
-				LogMessages.Add(MakeShareable(new FAutomationOutputMessage(TestResults.Errors[ErrorIndex], TEXT("Automation.Error"))));
-			}
-			for (int32 WarningIndex = 0; WarningIndex < TestResults.Warnings.Num(); ++WarningIndex)
+			const int32 NumOfPasses = Selection->GetNumResults(ClusterIndex);
+			for( int32 PassIndex = 0; PassIndex < NumOfPasses; ++PassIndex )
 			{
-				LogMessages.Add(MakeShareable(new FAutomationOutputMessage(TestResults.Warnings[WarningIndex], TEXT("Automation.Warning"))));
+				//get strings out of the report and populate the Log Messages
+				FAutomationTestResults TestResults = Selection->GetResults(ClusterIndex,PassIndex);
+
+				//no sense displaying device name if only one is available
+				if (NumOfPasses > 1)
+				{
+					FString PassHeader = LOCTEXT("TestPassHeader", "Pass:").ToString();
+					PassHeader += FString::Printf(TEXT("%i"),PassIndex+1);
+					LogMessages.Add(MakeShareable(new FAutomationOutputMessage(PassHeader, TEXT("Automation.Header"))));
+				}
+
+				for (int32 ErrorIndex = 0; ErrorIndex < TestResults.Errors.Num(); ++ErrorIndex)
+				{
+					LogMessages.Add(MakeShareable(new FAutomationOutputMessage(TestResults.Errors[ErrorIndex], TEXT("Automation.Error"))));
+				}
+				for (int32 WarningIndex = 0; WarningIndex < TestResults.Warnings.Num(); ++WarningIndex)
+				{
+					LogMessages.Add(MakeShareable(new FAutomationOutputMessage(TestResults.Warnings[WarningIndex], TEXT("Automation.Warning"))));
+				}
+				for (int32 LogIndex = 0; LogIndex < TestResults.Logs.Num(); ++LogIndex)
+				{
+					LogMessages.Add(MakeShareable(new FAutomationOutputMessage(TestResults.Logs[LogIndex], TEXT("Automation.Normal"))));
+				}
+				if ((TestResults.Warnings.Num() == 0) &&(TestResults.Errors.Num() == 0) && (Selection->GetState(ClusterIndex,PassIndex)==EAutomationState::Success))
+				{
+					LogMessages.Add(MakeShareable(new FAutomationOutputMessage(LOCTEXT( "AutomationTest_SuccessMessage", "Success" ).ToString(), TEXT("Automation.Normal"))));
+				}
+
+				LogMessages.Add(MakeShareable(new FAutomationOutputMessage(TEXT(""), TEXT("Log.Normal"))));
 			}
-			for (int32 LogIndex = 0; LogIndex < TestResults.Logs.Num(); ++LogIndex)
-			{
-				LogMessages.Add(MakeShareable(new FAutomationOutputMessage(TestResults.Logs[LogIndex], TEXT("Automation.Normal"))));
-			}
-			if ((TestResults.Warnings.Num() == 0) &&(TestResults.Errors.Num() == 0) && (Selection->GetState(ClusterIndex)==EAutomationState::Success))
-			{
-				LogMessages.Add(MakeShareable(new FAutomationOutputMessage(LOCTEXT( "AutomationTest_SuccessMessage", "Success" ).ToString(), TEXT("Automation.Normal"))));
-			}
-
-			LogMessages.Add(MakeShareable(new FAutomationOutputMessage(TEXT(""), TEXT("Log.Normal"))));
 		}
 		//rebuild UI
 		LogListView->RequestListRefresh();
@@ -670,7 +969,15 @@ TSharedRef<ITableRow> SAutomationWindow::OnGenerateWidgetForLog(TSharedPtr<FAuto
 
 FString SAutomationWindow::OnGetNumEnabledTestsString() const
 {
-	return FString::Printf(TEXT("%d"), AutomationController->GetEnabledTestsNum());
+	int32 NumPasses = AutomationController->GetNumPasses();
+	if( NumPasses > 1 )
+	{
+		return FString::Printf(TEXT("%d x%d"), AutomationController->GetEnabledTestsNum(),NumPasses);
+	}
+	else
+	{
+		return FString::Printf(TEXT("%d"), AutomationController->GetEnabledTestsNum());
+	}
 }
 
 
@@ -700,6 +1007,9 @@ void SAutomationWindow::OnRefreshTestCallback()
 
 	//rebuild the UI
 	TestTable->RequestTreeRefresh();
+
+	//update the background style
+	UpdateTestListBackgroundStyle();
 }
 
 
@@ -757,6 +1067,37 @@ void SAutomationWindow::FindWorkers()
 	MenuBar->SetEnabled( SessionIsValid );
 }
 
+void SAutomationWindow::HandleSessionManagerInstanceChanged()
+{
+	UpdateTestListBackgroundStyle();
+}
+
+void SAutomationWindow::UpdateTestListBackgroundStyle()
+{
+	TArray<ISessionInstanceInfoPtr> OutInstances;
+
+	if( ActiveSession.IsValid() )
+	{
+		ActiveSession->GetInstances(OutInstances);
+	}
+
+	TestBackgroundType = EAutomationTestBackgroundStyle::Unknown;
+
+	if( OutInstances.Num() > 0 )
+	{
+		FString FirstInstanceType = OutInstances[0]->GetInstanceType();
+
+		if( FirstInstanceType.Contains(TEXT("Editor")) )
+		{
+			TestBackgroundType = EAutomationTestBackgroundStyle::Editor;
+		}
+		else if( FirstInstanceType.Contains(TEXT("Game")) )
+		{
+			TestBackgroundType = EAutomationTestBackgroundStyle::Game;
+		}
+	}
+}
+
 
 FReply SAutomationWindow::RunTests()
 {
@@ -854,6 +1195,16 @@ void SAutomationWindow::OnToggleErrorFilter()
 	OnRefreshTestCallback();
 }
 
+void SAutomationWindow::OnChangeRepeatCount(int32 InNewValue)
+{
+	AutomationController->SetNumPasses(InNewValue);
+}
+
+int32 SAutomationWindow::GetRepeatCount() const
+{
+	return AutomationController->GetNumPasses();
+}
+
 
 FString SAutomationWindow::GetSmallIconExtension() const
 {
diff --git a/Engine/Source/Developer/AutomationWindow/Private/SAutomationWindow.h b/Engine/Source/Developer/AutomationWindow/Private/SAutomationWindow.h
--- a/Engine/Source/Developer/AutomationWindow/Private/SAutomationWindow.h
+++ b/Engine/Source/Developer/AutomationWindow/Private/SAutomationWindow.h
@@ -12,6 +12,17 @@ namespace AutomationTestWindowConstants
 	const FName Timing( TEXT("Timing") );
 }
 
+/** The type of background style to use for the test list widget */
+namespace EAutomationTestBackgroundStyle
+{
+	enum Type
+	{
+		Unknown,
+		Editor,
+		Game,
+	};
+}
+
 /**
  * Implements the main UI Window for hosting all automation tests.
  */
@@ -138,6 +149,21 @@ private:
 	static TSharedRef< SWidget > MakeAutomationWindowToolBar( const TSharedRef<FUICommandList>& InCommandList, TSharedPtr<class SAutomationWindow> InLevelEditor );
 	TSharedRef< SWidget > MakeAutomationWindowToolBar( const TSharedRef<FUICommandList>& InCommandList );
 
+	/**
+	 * Static: Creates the test options menu widget
+	 *
+	 * @return	New widget
+	 */
+	static TSharedRef< SWidget >GenerateTestsOptionsMenuContent( TWeakPtr<class SAutomationWindow> InAutomationWindow );
+	TSharedRef< SWidget > GenerateTestsOptionsMenuContent( );
+
+	/**
+	 * Creates a combo item for the preset list
+	 *
+	 * @return	New combo item widget
+	 */
+	TSharedRef<SWidget> GeneratePresetComboItem(TSharedPtr<FAutomationTestPreset> InItem);
+
 	/** 
 	 * Populates OutSearchStrings with the strings that should be used in searching
 	 * @param Report - the automation report to get a text description from
@@ -203,6 +229,14 @@ private:
 	bool IsErrorFilterOn() const;
 	/** Toggles filtering of tests based on error condition */
 	void OnToggleErrorFilter();
+	
+	/** Sets the number of times to repeat the tests */
+	void OnChangeRepeatCount(int32 InNewValue);
+	/** Returns the number of times to repeat the tests */
+	int32 GetRepeatCount() const;
+
+	/** Update the test list background style (Editor vs Game) */
+	void UpdateTestListBackgroundStyle();
 
 	/**
 	* Gets the extention for the small brush, if enabled
@@ -232,6 +266,13 @@ private:
 	*/
 	FString GetRunAutomationLabel() const;
 
+	/**
+	* Gets the brush to use for the test list background
+	*
+	* @return The brush to use depending on the value of TestBackgroundType
+	*/
+	const FSlateBrush* GetTestBackgroundBorderImage() const;
+
 	/**
 	* Recursively expand the tree nodes
 	*
@@ -256,6 +297,11 @@ private:
 	*/
 	void HandleSessionManagerSelectionChanged( const ISessionInfoPtr& SelectedSession );
 
+	/** 
+	* Called when the session manager updates an instances
+	*/
+	void HandleSessionManagerInstanceChanged();
+
 	/**
 	* Should the automation run button be enabled
 	*
@@ -284,6 +330,61 @@ private:
 	TSharedPtr<SWidget> HandleAutomationListContextMenuOpening( );
 #endif
 
+	/**
+	 * Handles the new preset button being clicked
+	 */
+	FReply HandleNewPresetClicked();
+
+	/**
+	 * Handles the save preset button being clicked
+	 */
+	FReply HandleSavePresetClicked();
+
+	/**
+	 * Handles the remove preset button being clicked
+	 */
+	FReply HandleRemovePresetClicked();
+
+	/**
+	* Should the add preset button be enabled
+	*/
+	bool IsAddButtonIsEnabled() const;
+
+	/**
+	* Should the save preset button be enabled
+	*/
+	bool IsSaveButtonIsEnabled() const;
+
+	/**
+	* Should the remove preset button be enabled
+	*/
+	bool IsRemoveButtonIsEnabled() const;
+
+	/**
+	* Handles if the preset combo box should be visible
+	*/
+	EVisibility HandlePresetComboVisibility( ) const;
+
+	/**
+	* Handles if the add preset text box should be visible
+	*/
+	EVisibility HandlePresetTextVisibility( ) const;
+
+	/**
+	* Called when the user commits the text in the add preset text box
+	*/
+	void HandlePresetTextCommited( const FText& CommittedText, ETextCommit::Type CommitType );
+
+	/**
+	* Called when the user selects a new preset from the preset combo box
+	*/
+	void HandlePresetChanged( TSharedPtr<FAutomationTestPreset> Item, ESelectInfo::Type SelectInfo );
+
+	/**
+	* Gets the text to display for the preset combo box
+	*/
+	FString GetPresetComboText() const;
+
 	/**
 	* Handle the copy button clicked in the command bar
 	*
@@ -365,4 +466,24 @@ private:
 
 	// Flag to acknowledge if the window is awaiting tests to display
 	bool bIsRequestingTests;
-};
+
+	// Which type of window style to use for the test background
+	EAutomationTestBackgroundStyle::Type TestBackgroundType;
+
+	// True if we are creating a new preset (The add preset text box is visible)
+	bool bAddingTestPreset;
+
+	// Holds a pointer to the preset manager
+	TSharedPtr<FAutomationTestPresetManager> TestPresetManager;
+
+	// Holds the currently selected preset
+	TSharedPtr<FAutomationTestPreset> SelectedPreset;
+
+	// Holds a pointer to the preset combo box widget
+	TSharedPtr< SComboBox< TSharedPtr<FAutomationTestPreset> > > PresetComboBox;
+
+	// Holds a pointer to the preset text box
+	TSharedPtr<SEditableTextBox> PresetTextBox;
+
+
+};
\ No newline at end of file
diff --git a/Engine/Source/Developer/SessionServices/Private/SessionManager.cpp b/Engine/Source/Developer/SessionServices/Private/SessionManager.cpp
--- a/Engine/Source/Developer/SessionServices/Private/SessionManager.cpp
+++ b/Engine/Source/Developer/SessionServices/Private/SessionManager.cpp
@@ -242,6 +242,7 @@ void FSessionManager::HandleEnginePongMessage( const FEngineServicePong& Message
 	if (Session.IsValid())
 	{
 		Session->UpdateFromMessage(Message, Context);
+		SessionInstanceUpdatedDelegate.Broadcast();
 	}
 }
 
diff --git a/Engine/Source/Developer/SessionServices/Private/SessionManager.h b/Engine/Source/Developer/SessionServices/Private/SessionManager.h
--- a/Engine/Source/Developer/SessionServices/Private/SessionManager.h
+++ b/Engine/Source/Developer/SessionServices/Private/SessionManager.h
@@ -73,6 +73,11 @@ public:
 		return SessionsUpdatedDelegate;
 	}
 
+	virtual FSimpleMulticastDelegate& OnSessionInstanceUpdated( ) OVERRIDE
+	{
+		return SessionInstanceUpdatedDelegate;
+	}
+
 	virtual void RemoveOwner( const FString& InOwner ) OVERRIDE;
 
 	virtual bool SelectSession( const ISessionInfoPtr& Session ) OVERRIDE;
@@ -164,6 +169,9 @@ private:
 	// Holds a delegate to be invoked when the session list was updated.
 	FSimpleMulticastDelegate SessionsUpdatedDelegate;
 
+	// Holds a delegate to be invoked when a session instance is updated.
+	FSimpleMulticastDelegate SessionInstanceUpdatedDelegate;
+
 	// Holds a delegate to be invoked when the widget ticks.
 	FTickerDelegate TickDelegate;
 
diff --git a/Engine/Source/Developer/SessionServices/Public/Interfaces/ISessionManager.h b/Engine/Source/Developer/SessionServices/Public/Interfaces/ISessionManager.h
--- a/Engine/Source/Developer/SessionServices/Public/Interfaces/ISessionManager.h
+++ b/Engine/Source/Developer/SessionServices/Public/Interfaces/ISessionManager.h
@@ -146,6 +146,13 @@ public:
 	 */
 	virtual FSimpleMulticastDelegate& OnSessionsUpdated( ) = 0;
 
+	/**
+	 * Returns a delegate that is executed when a session instance is updated.
+	 *
+	 * @return The delegate.
+	 */
+	virtual FSimpleMulticastDelegate& OnSessionInstanceUpdated( ) = 0;
+
 public:
 
 	/**
diff --git a/Engine/Source/Editor/EditorStyle/Private/SlateEditorStyle.cpp b/Engine/Source/Editor/EditorStyle/Private/SlateEditorStyle.cpp
--- a/Engine/Source/Editor/EditorStyle/Private/SlateEditorStyle.cpp
+++ b/Engine/Source/Editor/EditorStyle/Private/SlateEditorStyle.cpp
@@ -4988,6 +4988,15 @@ void FSlateEditorStyle::FStyle::SetupAutomationStyles()
 		Set( "AutomationWindow.SmokeTestFilter.Small", new IMAGE_BRUSH( "Automation/SmokeTestFilter", Icon20x20) );
 		Set( "AutomationWindow.DeveloperDirectoryContent.Small", new IMAGE_BRUSH( "Automation/DeveloperDirectoryContent", Icon20x20) );
 		Set( "AutomationWindow.VisualCommandlet.Small", new IMAGE_BRUSH( "Automation/VisualCommandlet", Icon20x20) );
+
+		//test preset icons
+		Set( "AutomationWindow.PresetNew", new IMAGE_BRUSH( "Icons/icon_add_40x", Icon16x16 ) );
+		Set( "AutomationWindow.PresetSave", new IMAGE_BRUSH( "Icons/icon_file_save_16px", Icon16x16 ) );
+		Set( "AutomationWindow.PresetRemove", new IMAGE_BRUSH( "Icons/icon_Cascade_DeleteLOD_40x", Icon16x16 ) );
+
+		//test backgrounds
+		Set( "AutomationWindow.GameGroupBorder", new BOX_BRUSH( "Automation/GameGroupBorder", FMargin(4.0f/16.0f) ) );
+		Set( "AutomationWindow.EditorGroupBorder", new BOX_BRUSH( "Automation/EditorGroupBorder", FMargin(4.0f/16.0f) ) );
 	}
 
 	// Launcher
diff --git a/Engine/Source/Editor/StaticMeshEditor/Private/StaticMeshAutomationTests.cpp b/Engine/Source/Editor/StaticMeshEditor/Private/StaticMeshAutomationTests.cpp
new file mode 100644
--- /dev/null
+++ b/Engine/Source/Editor/StaticMeshEditor/Private/StaticMeshAutomationTests.cpp
@@ -0,0 +1,186 @@
+// Copyright 1998-2014 Epic Games, Inc. All Rights Reserved.
+
+#include "StaticMeshEditorModule.h"
+
+#include "AutomationCommon.h"
+#include "SStaticMeshEditorViewport.h"
+#include "StaticMeshEditorViewportClient.h"
+
+namespace EditorViewButtonHelper
+{
+	/**
+	 * The types of buttons that will be toggled on and off, if new buttons are made that want to be added, 
+	 * all you need to do is add them to this list, and fill out the latent automation task below with how to toggle the button
+	 */
+	namespace EStaticMeshFlag
+	{
+		enum Type
+		{
+			Wireframe = 0,
+			Vert,
+			Grid,
+			Bounds,
+			Collision,
+			Pivot,
+			Normals,
+			Tangents,
+			Binormals,
+			UV,
+			Max //Do not go beyond this
+		};
+	}
+
+	/**
+	 * The struct we will be passing into the below Latent Automation Task
+	 */
+	struct PerformStaticMeshFlagParameters
+	{
+		FStaticMeshEditorViewportClient* ViewportClient;
+		EStaticMeshFlag::Type CommandType;
+	};
+
+	/**
+	 * Used for toggling a certain button on/off based on the passed in viewport client and button type.
+	 */
+	DEFINE_LATENT_AUTOMATION_COMMAND_ONE_PARAMETER(FPerformStaticMeshFlagToggle, PerformStaticMeshFlagParameters, AutomationParameters);
+
+	bool EditorViewButtonHelper::FPerformStaticMeshFlagToggle::Update()
+	{ 
+		switch(AutomationParameters.CommandType)
+		{
+			case EStaticMeshFlag::Wireframe:
+				if (AutomationParameters.ViewportClient->GetViewMode() != VMI_Wireframe)
+				{
+					AutomationParameters.ViewportClient->SetViewMode(VMI_Wireframe);
+				}
+				else
+				{
+					AutomationParameters.ViewportClient->SetViewMode(VMI_Lit);
+				}
+				break;
+			case EStaticMeshFlag::Vert:
+				if (!AutomationParameters.ViewportClient->EngineShowFlags.VertexColors)
+				{
+					AutomationParameters.ViewportClient->EngineShowFlags.VertexColors = true;
+					AutomationParameters.ViewportClient->EngineShowFlags.Lighting = false;
+				}
+				else
+				{
+					AutomationParameters.ViewportClient->EngineShowFlags.VertexColors = false;
+					AutomationParameters.ViewportClient->EngineShowFlags.Lighting = true;
+				}
+				break;
+			case EStaticMeshFlag::Grid:
+				AutomationParameters.ViewportClient->SetShowGrid();
+				break;
+			case EStaticMeshFlag::Bounds:
+				if (AutomationParameters.ViewportClient->EngineShowFlags.Bounds)
+				{
+					AutomationParameters.ViewportClient->EngineShowFlags.Bounds = false;
+				}
+				else
+				{
+					AutomationParameters.ViewportClient->EngineShowFlags.Bounds = true;
+				}
+				break;
+			case EStaticMeshFlag::Collision:
+				AutomationParameters.ViewportClient->SetShowCollision();
+				break;
+			case EStaticMeshFlag::Pivot:
+				AutomationParameters.ViewportClient->SetShowPivot();
+				break;
+			case EStaticMeshFlag::Normals:
+				AutomationParameters.ViewportClient->SetShowNormals();
+				break;
+			case EStaticMeshFlag::Tangents:
+				AutomationParameters.ViewportClient->SetShowTangents();
+				break;
+			case EStaticMeshFlag::Binormals:
+				AutomationParameters.ViewportClient->SetShowBinormals();
+				break;
+			case EStaticMeshFlag::UV:
+				AutomationParameters.ViewportClient->SetDrawUVOverlay();
+				break;
+			default:
+				//Break out immediately.
+				break;
+		}
+		return true;
+	}
+
+	/**
+	* Close all asset editors
+	*/
+	DEFINE_LATENT_AUTOMATION_COMMAND(FCloseAllAssetEditorsCommand);
+
+	bool FCloseAllAssetEditorsCommand::Update()
+	{
+		FAssetEditorManager::Get().CloseAllAssetEditors();
+
+		return true;
+	}
+}
+
+/**
+ * Static mesh editor test
+ */
+IMPLEMENT_SIMPLE_AUTOMATION_TEST( FStaticMeshEditorTest, "Editor.Content.Static Mesh Editor Test", EAutomationTestFlags::ATF_Editor )
+
+/**
+ * Take screenshots of the SME window with each of the toolbar buttons toggled separately
+ */
+bool FStaticMeshEditorTest::RunTest(const FString& Parameters)
+{
+	TArray<FString> ButtonNamesToTestFor;
+	FString LoadedObjectType = TEXT("EditorCylinder");
+	FString MeshName;
+	EditorViewButtonHelper::PerformStaticMeshFlagParameters AutomationParameters;
+	WindowScreenshotParameters WindowParameters;
+
+	//Pull from .ini the names of the buttons we want to test for
+	GConfig->GetString( TEXT("AutomationTesting"), TEXT("EditorViewButtonsObject"), LoadedObjectType, GEngineIni);
+
+	//Open static mesh in the editor and save out the window for the screenshots
+	UObject* EditorMesh = (UStaticMesh*)StaticLoadObject(UStaticMesh::StaticClass(),
+		NULL,
+		*FString::Printf(TEXT("/Engine/EditorMeshes/%s.%s"), *LoadedObjectType, *LoadedObjectType),
+		NULL,
+		LOAD_None,
+		NULL);
+	FAssetEditorManager::Get().OpenEditorForAsset(EditorMesh);
+	WindowParameters.CurrentWindow = FSlateApplication::Get().GetActiveTopLevelWindow();
+
+	//Grab the last opened Viewport (aka the asset manager)
+	AutomationParameters.ViewportClient = static_cast<FStaticMeshEditorViewportClient*>(GEditor->AllViewportClients[GEditor->AllViewportClients.Num() - 1]);
+	AutomationParameters.CommandType = EditorViewButtonHelper::EStaticMeshFlag::Wireframe; //Start it off at Wireframe because that is the first button we are attempting to press
+
+	if (AutomationParameters.ViewportClient)
+	{
+		const FString BaseFileName = FPaths::ScreenShotDir() / TEXT("StaticMeshEditor");
+		WindowParameters.ScreenshotName = FString::Printf(TEXT("%s00.bmp"), *BaseFileName);
+
+		//Wait for the window to load and then take the initial screen shot		
+		ADD_LATENT_AUTOMATION_COMMAND(FWaitLatentCommand(0.5f));
+		ADD_LATENT_AUTOMATION_COMMAND(FTakeEditorScreenshotCommand(WindowParameters));
+		ADD_LATENT_AUTOMATION_COMMAND(FWaitLatentCommand(0.1f));
+
+		for(auto CurrentCommandType = 0; CurrentCommandType < EditorViewButtonHelper::EStaticMeshFlag::Max; ++CurrentCommandType)
+		{
+			WindowParameters.ScreenshotName = FString::Printf(TEXT("%s%02i.bmp"), *BaseFileName, CurrentCommandType + 1);
+
+			//Toggle the command for the button, take a screenshot, and then re-toggle the command for the button
+			ADD_LATENT_AUTOMATION_COMMAND(EditorViewButtonHelper::FPerformStaticMeshFlagToggle(AutomationParameters));
+			ADD_LATENT_AUTOMATION_COMMAND(FWaitLatentCommand(0.1f));
+			ADD_LATENT_AUTOMATION_COMMAND(FTakeEditorScreenshotCommand(WindowParameters));
+			ADD_LATENT_AUTOMATION_COMMAND(FWaitLatentCommand(0.1f));
+			ADD_LATENT_AUTOMATION_COMMAND(EditorViewButtonHelper::FPerformStaticMeshFlagToggle(AutomationParameters));
+			
+			//Increment to the next enum, so we can loop through all the types
+			AutomationParameters.CommandType = (EditorViewButtonHelper::EStaticMeshFlag::Type)((int)AutomationParameters.CommandType + 1);
+		}
+
+		ADD_LATENT_AUTOMATION_COMMAND(EditorViewButtonHelper::FCloseAllAssetEditorsCommand());
+	}
+
+	return true;
+}
diff --git a/Engine/Source/Editor/UnrealEd/Private/Tests/EditorAutomationTests.cpp b/Engine/Source/Editor/UnrealEd/Private/Tests/EditorAutomationTests.cpp
--- a/Engine/Source/Editor/UnrealEd/Private/Tests/EditorAutomationTests.cpp
+++ b/Engine/Source/Editor/UnrealEd/Private/Tests/EditorAutomationTests.cpp
@@ -17,25 +17,9 @@
 #include "BlueprintGraphClasses.h"
 
 #include "SlateWordWrapper.h"
+#include "AutomationCommon.h"
 
-
-DEFINE_LOG_CATEGORY_STATIC(LogEditorAutomationTests, Log, All);
-
-
-/**
- * Wait for the given amount of time
- */
-DEFINE_LATENT_AUTOMATION_COMMAND_ONE_PARAMETER(FWaitLatentCommand, float, Duration);
-
-bool FWaitLatentCommand::Update()
-{
-	float NewTime = FPlatformTime::Seconds();
-	if (NewTime - StartTime >= Duration)
-	{
-		return true;
-	}
-	return false;
-}
+//DEFINE_LOG_CATEGORY_STATIC(LogEditorAutomationTests, Log, All);
 
 /**
  * Start PIE session
@@ -82,6 +66,7 @@ bool FOpenEditorForAssetCommand::Update()
 
 	return true;
 }
+
 /**
  * Close all asset editors
  */
@@ -152,6 +137,7 @@ public:
 			}
 		}
 	}
+
 	
 };
 
@@ -1845,3 +1831,62 @@ bool FTraceAllTimelinesAutomationTest::RunTest(const FString& BlueprintName)
 	return bPassed;
 }
 
+/**
+ * StaticMeshUVsTest
+ */
+IMPLEMENT_SIMPLE_AUTOMATION_TEST( FStaticMeshUVsTest, "Editor.Content.Static Mesh UVs", EAutomationTestFlags::ATF_Editor )
+
+/**
+ * Find all static meshes and check the lightmap UVs
+ */
+bool FStaticMeshUVsTest::RunTest(const FString& Parameters)
+{
+	TArray<FString> MissingUVMessages;
+	TArray<FString> BadUVMessages;
+	TArray<FString> ValidUVMessages;
+	
+	//Create our package list
+	TArray<FString> ContentPackages;
+	FPackageName::FindPackagesInDirectory(ContentPackages, *FPaths::EngineContentDir());
+	FPackageName::FindPackagesInDirectory(ContentPackages, *FPaths::GameContentDir());
+	UE_LOG(LogEditorAutomationTests, Log, TEXT("Found %i Packages."), ContentPackages.Num() );
+
+	for( int32 PackageID = 0; PackageID < ContentPackages.Num(); ++PackageID)
+	{
+		const FString PackageName = ContentPackages[PackageID];
+
+		UPackage* Package = NULL;
+		{
+			SetSuppressLogs(true);	
+			Package = LoadPackage( NULL, *PackageName, LOAD_NoWarn | LOAD_Quiet );
+			SetSuppressLogs(false);	
+		}
+		if (Package != NULL)
+		{
+			for (FObjectIterator ObjIt; ObjIt; ++ObjIt)
+			{
+				UObject* Object = *ObjIt;
+				if (Object != NULL)
+				{
+					//Call CheckLightMapUVs on all static meshes in this package
+					if( Object->IsA(UStaticMesh::StaticClass()) && Object->IsIn(Package) )
+					{
+						UStaticMesh::CheckLightMapUVs((UStaticMesh*)Object,MissingUVMessages,BadUVMessages,ValidUVMessages,true);
+					}
+				}
+			}
+		}
+		else
+		{
+			UE_LOG(LogEditorAutomationTests, Error, TEXT("Error loading package %s."), *PackageName );
+		}
+		
+		//Collect garbage on every Nth package to keep memory usage down.
+		if( ( (++PackageID % 10) == 0 ) )
+		{
+			CollectGarbage(RF_Native);
+		}
+	}
+	CollectGarbage(RF_Native);
+	return true;
+}
diff --git a/Engine/Source/Runtime/AutomationWorker/Private/AutomationWorkerModule.cpp b/Engine/Source/Runtime/AutomationWorker/Private/AutomationWorkerModule.cpp
--- a/Engine/Source/Runtime/AutomationWorker/Private/AutomationWorkerModule.cpp
+++ b/Engine/Source/Runtime/AutomationWorker/Private/AutomationWorkerModule.cpp
@@ -276,6 +276,7 @@ void FAutomationWorkerModule::HandleResetTests( const FAutomationWorkerResetTest
 
 void FAutomationWorkerModule::HandleRequestTestsMessage( const FAutomationWorkerRequestTests& Message, const IMessageContextRef& Context )
 {
+	FAutomationTestFramework::GetInstance().LoadTestModules();
 	FAutomationTestFramework::GetInstance().SetDeveloperDirectoryIncluded(Message.DeveloperDirectoryIncluded);
 	FAutomationTestFramework::GetInstance().SetVisualCommandletFilter(Message.VisualCommandletFilterOn);
 	FAutomationTestFramework::GetInstance().GetValidTestNames( TestInfo );
diff --git a/Engine/Source/Runtime/Core/Private/Misc/AutomationTest.cpp b/Engine/Source/Runtime/Core/Private/Misc/AutomationTest.cpp
--- a/Engine/Source/Runtime/Core/Private/Misc/AutomationTest.cpp
+++ b/Engine/Source/Runtime/Core/Private/Misc/AutomationTest.cpp
@@ -5,6 +5,7 @@
 =============================================================================*/
 
 #include "CorePrivate.h"
+#include "ModuleManager.h"
 
 DEFINE_LOG_CATEGORY_STATIC(LogAutomationTest, Warning, All);
 
@@ -301,6 +302,40 @@ bool FAutomationTestFramework::ExecuteNetworkCommands()
 	return !bHadAnyNetworkCommands;
 }
 
+void FAutomationTestFramework::LoadTestModules( )
+{
+	const bool bRunningEditor = GIsEditor && !IsRunningCommandlet();
+
+	if( !bRunningSmokeTests )
+	{
+		TArray<FString> EngineTestModules;
+		GConfig->GetArray( TEXT("/Script/Engine.AutomationTestSettings"), TEXT("EngineTestModules"), EngineTestModules, GEngineIni);
+		//Load any engine level modules.
+		for( int32 EngineModuleId = 0; EngineModuleId < EngineTestModules.Num(); ++EngineModuleId)
+		{
+			const FName ModuleName = FName(*EngineTestModules[EngineModuleId]);
+			if (!FModuleManager::Get().IsModuleLoaded(ModuleName))
+			{
+				FModuleManager::Get().LoadModule(ModuleName);
+			}
+		}
+		//Load any editor modules.
+		if( bRunningEditor )
+		{
+			TArray<FString> EditorTestModules;
+			GConfig->GetArray( TEXT("/Script/Engine.AutomationTestSettings"), TEXT("EditorTestModules"), EditorTestModules, GEngineIni);
+			for( int32 EditorModuleId = 0; EditorModuleId < EditorTestModules.Num(); ++EditorModuleId )
+			{
+				const FName ModuleName = FName(*EditorTestModules[EditorModuleId]);
+				if (!FModuleManager::Get().IsModuleLoaded(ModuleName))
+				{
+					FModuleManager::Get().LoadModule(ModuleName);
+				}
+			}
+		}
+	}
+}
+
 
 void FAutomationTestFramework::GetValidTestNames( TArray<FAutomationTestInfo>& TestInfo ) const
 {
@@ -581,19 +616,28 @@ void FAutomationTestBase::ClearExecutionInfo()
 
 void FAutomationTestBase::AddError( const FString& InError )
 {
-	ExecutionInfo.Errors.Add( InError );
+	if( !bSuppressLogs )
+	{
+		ExecutionInfo.Errors.Add( InError );
+	}
 }
 
 
 void FAutomationTestBase::AddWarning( const FString& InWarning )
 {
-	ExecutionInfo.Warnings.Add( InWarning );
+	if( !bSuppressLogs )
+	{
+		ExecutionInfo.Warnings.Add( InWarning );
+	}
 }
 
 
 void FAutomationTestBase::AddLogItem( const FString& InLogItem )
 {
-	ExecutionInfo.LogItems.Add( InLogItem );
+	if( !bSuppressLogs )
+	{
+		ExecutionInfo.LogItems.Add( InLogItem );
+	}
 }
 
 
diff --git a/Engine/Source/Runtime/Core/Public/Misc/AutomationTest.h b/Engine/Source/Runtime/Core/Public/Misc/AutomationTest.h
--- a/Engine/Source/Runtime/Core/Public/Misc/AutomationTest.h
+++ b/Engine/Source/Runtime/Core/Public/Misc/AutomationTest.h
@@ -459,6 +459,11 @@ public:
 	 */
 	bool ExecuteNetworkCommands();
 
+	/**
+	 * Load any modules that are not loaded by default and have test classes in them
+	 */
+	void LoadTestModules();
+
 	/**
 	 * Populates the provided array with the names of all tests in the framework that are valid to run for the current
 	 * application settings.
@@ -641,6 +646,7 @@ public:
 	FAutomationTestBase( const FString& InName, const bool bInComplexTask )
 		: bComplexTask( bInComplexTask )
 		, TestName( InName )
+		, bSuppressLogs( false )
 	{
 		// Register the newly created automation test into the automation testing framework
 		FAutomationTestFramework::GetInstance().RegisterAutomationTest( InName, this );
@@ -728,6 +734,15 @@ public:
 		return bComplexTask;
 	}
 
+	/**
+	 * Used to suppress / unsuppress logs.
+	 *
+	 * @param bNewValue - True if you want to suppress logs.  False to unsuppress.
+	 */
+	void SetSuppressLogs(bool bNewValue)
+	{
+		bSuppressLogs = bNewValue;
+	}
 
 public:
 
@@ -921,6 +936,9 @@ protected:
 	//Flag to indicate if this is a complex task
 	bool bComplexTask;
 
+	/** Flag to suppress logs */
+	bool bSuppressLogs;
+
 	/** Name of the test */
 	FString TestName;
 
@@ -960,6 +978,29 @@ class CommandName : public IAutomationLatentCommand \
 	ParamType ParamName; \
 }
 
+#define DEFINE_ENGINE_LATENT_AUTOMATION_COMMAND(CommandName)	\
+class ENGINE_API CommandName : public IAutomationLatentCommand \
+	{ \
+	public: \
+	virtual ~CommandName() \
+		{} \
+		virtual bool Update() OVERRIDE; \
+}
+
+#define DEFINE_ENGINE_LATENT_AUTOMATION_COMMAND_ONE_PARAMETER(CommandName,ParamType,ParamName)	\
+class ENGINE_API CommandName : public IAutomationLatentCommand \
+	{ \
+	public: \
+	CommandName(ParamType InputParam) \
+	: ParamName(InputParam) \
+		{} \
+		virtual ~CommandName() \
+		{} \
+		virtual bool Update() OVERRIDE; \
+	private: \
+	ParamType ParamName; \
+}
+
 //macro to simply the syntax for enqueueing a latent command
 #define ADD_LATENT_AUTOMATION_COMMAND(ClassDeclaration) FAutomationTestFramework::GetInstance().EnqueueLatentCommand(MakeShareable(new ClassDeclaration));
 
diff --git a/Engine/Source/Runtime/Engine/Classes/Tests/AutomationTestSettings.h b/Engine/Source/Runtime/Engine/Classes/Tests/AutomationTestSettings.h
--- a/Engine/Source/Runtime/Engine/Classes/Tests/AutomationTestSettings.h
+++ b/Engine/Source/Runtime/Engine/Classes/Tests/AutomationTestSettings.h
@@ -28,6 +28,39 @@ struct FOpenTestAsset
 	bool		bSkipTestWhenUnAttended;
 };
 
+/**
+ * Structure for defining an external tool
+ */
+USTRUCT()
+struct FExternalToolDefinition
+{
+	GENERATED_USTRUCT_BODY()
+
+	/* The name of the tool / test. */
+	UPROPERTY(config, EditAnywhere, Category=ExternalTools)
+	FString ToolName;
+
+	/* The executable to run. */
+	UPROPERTY(config, EditAnywhere, Category=ExternalTools, meta=(FilePathFilter = "*"))
+	FFilePath ExecutablePath;
+
+	/* The command line options to pass to the executable. */
+	UPROPERTY(config, EditAnywhere, Category=ExternalTools)
+	FString CommandLineOptions;
+
+	/* The working directory for the new process. */
+	UPROPERTY(config, EditAnywhere, Category=ExternalTools)
+	FString WorkingDirectory;
+
+	/* If set, look for scripts with this extension. */
+	UPROPERTY(config, EditAnywhere, Category=ExternalTools)
+	FString ScriptExtension;
+
+	/* If the ScriptExtension is set, look here for the script files. */
+	UPROPERTY(config, EditAnywhere, Category=ExternalTools)
+	FString ScriptDirectory;
+};
+
 /**
  * Implements the Editor's user settings.
  */
@@ -41,7 +74,7 @@ public:
 	/**
 	 * The directory to which the packaged project will be copied.
 	 */
-	UPROPERTY(config, EditAnywhere, Category=Automation)
+	UPROPERTY(config, EditAnywhere, Category=Automation, meta=(FilePathFilter = "umap"))
 	FFilePath AutomationTestmap;
 
 	/**
@@ -49,4 +82,22 @@ public:
 	 */
 	UPROPERTY(EditAnywhere, config, Category=Automation)
 	TArray<FOpenTestAsset> TestAssetsToOpen;
+
+	/**
+	 * Modules to load that have engine tests
+	 */
+	UPROPERTY(EditAnywhere, config, Category=Automation)
+	TArray<FString> EngineTestModules;
+
+	/**
+	 * Modules to load that have editor tests
+	 */
+	UPROPERTY(EditAnywhere, config, Category=Automation)
+	TArray<FString> EditorTestModules;
+
+	/**
+	 * Asset to test for open in automation process
+	 */
+	UPROPERTY(EditAnywhere, config, Category=ExternalTools)
+	TArray<FExternalToolDefinition> ExternalTools;
 };
diff --git a/Engine/Source/Runtime/Engine/Private/Tests/ExternalToolTest.cpp b/Engine/Source/Runtime/Engine/Private/Tests/ExternalToolTest.cpp
new file mode 100644
--- /dev/null
+++ b/Engine/Source/Runtime/Engine/Private/Tests/ExternalToolTest.cpp
@@ -0,0 +1,173 @@
+// Copyright 1998-2014 Epic Games, Inc. All Rights Reserved.
+
+#include "EnginePrivate.h"
+#include "EngineClasses.h"
+#include "AutomationTestCommon.h"
+
+DEFINE_LOG_CATEGORY_STATIC(LogBaseAutomationTests, Log, All);
+
+/**
+ * Holds the process info for an external tool
+ */
+namespace ExternalProcessHelper
+{
+	struct ExternalProcessInfo
+	{
+		ExternalProcessInfo()
+			: ProcessHandle(NULL)
+			, ReadPipe(NULL)
+			, WritePipe(NULL)
+		{
+		}
+		// Holds the process handle
+		FProcHandle		ProcessHandle;
+		// Holds the read pipe.
+		void* ReadPipe;
+		// Holds the write pipe.
+		void* WritePipe;
+	};
+}
+/**
+ * Wait for an external process to finish
+ */
+DEFINE_LATENT_AUTOMATION_COMMAND_ONE_PARAMETER(FWaitForProcessToCompleteLatentCommand, ExternalProcessHelper::ExternalProcessInfo, ProcessInfo);
+
+bool FWaitForProcessToCompleteLatentCommand::Update()
+{
+	if( ProcessInfo.ProcessHandle.IsValid() )
+	{
+		//Check for any output
+		FString ProcessOutput = FPlatformProcess::ReadPipe(ProcessInfo.ReadPipe);
+		if( ProcessOutput.Len() > 0 )
+		{
+			TArray<FString> OutputLines;
+			ProcessOutput.ParseIntoArray(&OutputLines, TEXT("\n"), false);
+			for( int32 i = 0; i < OutputLines.Num(); i++ )
+			{
+				UE_LOG(LogBaseAutomationTests, Log, TEXT("%s"), *OutputLines[i]);
+			}
+		}
+
+		//We aren't done until the process stops
+		if( FPlatformProcess::IsProcRunning(ProcessInfo.ProcessHandle) )
+		{
+			return false;
+		}
+
+		FPlatformProcess::ClosePipe(ProcessInfo.ReadPipe,ProcessInfo.WritePipe);
+		ProcessInfo.ReadPipe = NULL;
+		ProcessInfo.WritePipe = NULL;
+
+		//Check for any return codes
+		int32 ReturnCode;
+		if( !FPlatformProcess::GetProcReturnCode(ProcessInfo.ProcessHandle,&ReturnCode) )
+		{
+			ReturnCode = -1;
+		}
+
+		if( ReturnCode != 0 )
+		{
+			UE_LOG(LogBaseAutomationTests, Error, TEXT("External Tool failed with error code: %i"),ReturnCode);
+		}
+	}
+
+	return true;
+}
+
+/**
+ * RunExternalToolTest
+ */
+IMPLEMENT_COMPLEX_AUTOMATION_TEST( FRunExternalToolTest, "External", EAutomationTestFlags::ATF_None )
+
+/** 
+ * Find all the external too commands
+ */
+void FRunExternalToolTest::GetTests(TArray<FString>& OutBeautifiedNames, TArray <FString>& OutTestCommands) const
+{
+	UAutomationTestSettings const* AutomationTestSettings = GetDefault<UAutomationTestSettings>();
+	check(AutomationTestSettings);
+
+	for( int32 i = 0; i < AutomationTestSettings->ExternalTools.Num(); i++ )
+	{
+		const FExternalToolDefinition& ToolIt = AutomationTestSettings->ExternalTools[i];
+		if( ToolIt.ExecutablePath.FilePath.Len() > 0 )
+		{
+			// Check for multiple scripts if we have a script extension to look for
+			if( ToolIt.ScriptExtension.Len() > 0 && ToolIt.ScriptDirectory.Len() > 0 )
+			{
+				TArray<FString> ScriptFileNames;
+
+				IFileManager::Get().FindFiles(ScriptFileNames, *(ToolIt.ScriptDirectory / TEXT("*.") + ToolIt.ScriptExtension), true, false);
+
+				for (TArray<FString>::TConstIterator ScriptIt(ScriptFileNames); ScriptIt; ++ScriptIt)
+				{
+					const FString ScriptName = FString::Printf(TEXT("%s: %s"),*ToolIt.ToolName,*FPaths::GetBaseFilename(*ScriptIt));
+					OutBeautifiedNames.Add(*ScriptName);
+
+					const FString SctiptCommand = FString::Printf(TEXT("%s;%s;%s"),*ToolIt.ExecutablePath.FilePath,**ScriptIt,ToolIt.WorkingDirectory.Len() > 0 ? *ToolIt.WorkingDirectory : *ToolIt.ScriptDirectory);
+					OutTestCommands.Add(*SctiptCommand);
+				}
+			}
+			else
+			{
+				OutBeautifiedNames.Add(*ToolIt.ToolName);
+				const FString ToolCommand = FString::Printf(TEXT("%s;%s;%s"),*ToolIt.ExecutablePath.FilePath,*ToolIt.CommandLineOptions,*ToolIt.WorkingDirectory);
+				OutTestCommands.Add(*ToolCommand);
+			}
+		}
+	}
+}
+
+/** 
+ * Launch the external tool as a process and wait for it to complete
+ *
+ * @param Parameters - The command to launch the external tool
+ * @return	TRUE if the test was successful, FALSE otherwise
+ */
+bool FRunExternalToolTest::RunTest(const FString& Parameters)
+{
+	TArray<FString> Commands;
+	Parameters.ParseIntoArray(&Commands,TEXT(";"),false);
+
+	if( Commands.Num() != 3 )
+	{
+		UE_LOG(LogBaseAutomationTests, Error, TEXT("ERROR Parsing commands for external tool: %s"),*Parameters);
+		return false;
+	}
+
+	const FString& Executable       = Commands[0];
+	const FString& Options          = Commands[1];
+	const FString& WorkingDirectory = Commands[2];
+
+	//Create the pipes that we can use to read the output of the process
+	ExternalProcessHelper::ExternalProcessInfo ProcessInfo;
+	if (!FPlatformProcess::CreatePipe(ProcessInfo.ReadPipe, ProcessInfo.WritePipe))
+	{
+		return false;
+	}
+	
+	//Create the new process
+	ProcessInfo.ProcessHandle = FPlatformProcess::CreateProc(*Executable,
+															*Options,
+															true,
+															false,
+															false,
+															NULL,
+															0,
+															WorkingDirectory.Len() > 0 ? *WorkingDirectory : NULL,
+															ProcessInfo.WritePipe);
+
+	//Check that we got a valid process handle
+	if( ProcessInfo.ProcessHandle.IsValid() )
+	{
+		ADD_LATENT_AUTOMATION_COMMAND(FWaitForProcessToCompleteLatentCommand(ProcessInfo));
+	}
+	else
+	{
+		UE_LOG(LogBaseAutomationTests, Error, TEXT("Failed to launch executable (%s) for external tool %s"),*Executable, *Parameters);
+		FPlatformProcess::ClosePipe(ProcessInfo.ReadPipe, ProcessInfo.WritePipe);
+		return false;
+	}
+
+	return true;
+}
\ No newline at end of file
diff --git a/Engine/Source/Runtime/Engine/Public/Tests/AutomationCommon.cpp b/Engine/Source/Runtime/Engine/Public/Tests/AutomationCommon.cpp
new file mode 100644
--- /dev/null
+++ b/Engine/Source/Runtime/Engine/Public/Tests/AutomationCommon.cpp
@@ -0,0 +1,52 @@
+// Copyright 1998-2014 Epic Games, Inc. All Rights Reserved.
+
+#include "EnginePrivate.h"
+#include "AutomationCommon.h"
+
+///////////////////////////////////////////////////////////////////////
+// Common Latent commands which are used across test type. I.e. Engine, Network, etc...
+
+static void SaveWindowAsScreenshot(TSharedRef<SWindow> Window, const FString& FileName)
+{
+	TSharedRef<SWidget> WindowRef = Window;
+
+	TArray<FColor> OutImageData;
+	FIntVector OutImageSize;
+	FSlateApplication::Get().TakeScreenshot(WindowRef,OutImageData,OutImageSize);
+
+	FFileHelper::CreateBitmap(*FileName, OutImageSize.X, OutImageSize.Y, OutImageData.GetTypedData());
+}
+
+static void SaveEditorWindowsAsScreenshots(const FString& BaseFileName)
+{
+	TArray<TSharedRef<SWindow> > ChildWindows;
+	FSlateApplication::Get().GetAllVisibleWindowsOrdered(ChildWindows);
+
+	for( int32 i = 0; i < ChildWindows.Num(); ++i )
+	{
+		const FString BMPFileName = FString::Printf(TEXT("%s%02i.bmp"), *BaseFileName, i+1);
+		SaveWindowAsScreenshot(ChildWindows[i],BMPFileName);
+	}
+}
+
+bool FWaitLatentCommand::Update()
+{
+	float NewTime = FPlatformTime::Seconds();
+	if (NewTime - StartTime >= Duration)
+	{
+		return true;
+	}
+	return false;
+}
+
+bool FTakeActiveEditorScreenshotCommand::Update()
+{
+	SaveWindowAsScreenshot(FSlateApplication::Get().GetActiveTopLevelWindow().ToSharedRef(),ScreenshotName);
+	return true;
+}
+
+bool FTakeEditorScreenshotCommand::Update()
+{
+	SaveWindowAsScreenshot(ScreenshotParameters.CurrentWindow.ToSharedRef(), ScreenshotParameters.ScreenshotName);
+	return true;
+}
diff --git a/Engine/Source/Runtime/Engine/Public/Tests/AutomationCommon.h b/Engine/Source/Runtime/Engine/Public/Tests/AutomationCommon.h
new file mode 100644
--- /dev/null
+++ b/Engine/Source/Runtime/Engine/Public/Tests/AutomationCommon.h
@@ -0,0 +1,34 @@
+// Copyright 1998-2014 Epic Games, Inc. All Rights Reserved.
+
+#pragma once
+
+#include "EnginePrivate.h"
+
+///////////////////////////////////////////////////////////////////////
+// Common Latent commands which are used across test type. I.e. Engine, Network, etc...
+
+DEFINE_LOG_CATEGORY_STATIC(LogEditorAutomationTests, Log, All);
+
+/**
+ * Parameters to the Latent Automation command FTakeEditorScreenshotCommand
+ */
+struct WindowScreenshotParameters
+{
+	FString ScreenshotName;
+	TSharedPtr<SWindow> CurrentWindow;
+};
+
+/**
+ * Wait for the given amount of time
+ */
+DEFINE_ENGINE_LATENT_AUTOMATION_COMMAND_ONE_PARAMETER(FWaitLatentCommand, float, Duration);
+
+/**
+ * Take a screenshot of the active window
+ */
+DEFINE_ENGINE_LATENT_AUTOMATION_COMMAND_ONE_PARAMETER(FTakeActiveEditorScreenshotCommand, FString, ScreenshotName);
+
+/**
+ * Take a screenshot of the active window
+ */
+DEFINE_ENGINE_LATENT_AUTOMATION_COMMAND_ONE_PARAMETER(FTakeEditorScreenshotCommand, WindowScreenshotParameters, ScreenshotParameters);
