Merging //depot/Partners/Zombie/UE4-iOS/...to //depot/UE4/...

Automation: Milestone 3 and fixes for milestone 2

 -Added new EditorShot command to take editor screenshots
 -Made the ExterialTool options use the dierctory picker
 -Fixed a bug with the StaticMeshEditor test if the tutorial pops up
 -StaticMeshUV test now uses async package loading
 -Added a display every Nth screenshot to the screenshot compre tab
 -Added support to disable screenshots or request full size ones
 -Changed the screenshots so they save to the game directory even if you run the frontend by itself
 -Added screenshot support for the load all maps tests

diff --git a/Engine/Source/Developer/AutomationController/Private/AutomationControllerManager.h b/Engine/Source/Developer/AutomationController/Private/AutomationControllerManager.h
--- a/Engine/Source/Developer/AutomationController/Private/AutomationControllerManager.h
+++ b/Engine/Source/Developer/AutomationController/Private/AutomationControllerManager.h
@@ -41,6 +41,26 @@ public:
 		return NumTestPasses;
 	}
 
+	virtual bool IsUsingFullSizeScreenshots() const OVERRIDE
+	{
+		return bScreenshotsEnabled && bRequestFullScreenScreenshots;
+	}
+
+	virtual void SetUsingFullSizeScreenshots( const bool bNewValue ) OVERRIDE
+	{
+		bRequestFullScreenScreenshots = bNewValue;
+	}
+
+	virtual bool AreScreenshotsEnabled() const OVERRIDE
+	{
+		return bScreenshotsEnabled;
+	}
+
+	virtual void SetScreenshotsEnabled( const bool bNewValue ) OVERRIDE
+	{
+		bScreenshotsEnabled = bNewValue;
+	}
+
 	virtual void SetFilter( TSharedPtr< AutomationFilterCollection > InFilter ) OVERRIDE
 	{
 		ReportManager.SetFilter( InFilter );
@@ -339,6 +359,12 @@ private:
 
 	/** The current test pass we are on */
 	int32 CurrentTestPass;
+
+	/** If screenshots are enabled */
+	bool bScreenshotsEnabled;
+
+	/** If we should request full screen screen shots */
+	bool bRequestFullScreenScreenshots;
 private:
 
 	// Holds a delegate that is invoked when the controller shuts down.
diff --git a/Engine/Source/Developer/AutomationController/Private/AutomationControllerManger.cpp b/Engine/Source/Developer/AutomationController/Private/AutomationControllerManger.cpp
--- a/Engine/Source/Developer/AutomationController/Private/AutomationControllerManger.cpp
+++ b/Engine/Source/Developer/AutomationController/Private/AutomationControllerManger.cpp
@@ -118,6 +118,8 @@ void FAutomationControllerManager::Init()
 {
 	AutomationTestState = EAutomationControllerModuleState::Disabled;
 	bTestResultsAvailable = false;
+	bScreenshotsEnabled = true;
+	bRequestFullScreenScreenshots = false;
 }
 
 
@@ -212,7 +214,7 @@ void FAutomationControllerManager::ExecuteNextTask( int32 ClusterIndex, OUT bool
 							FMessageAddress DeviceAddress = DeviceAddresses[AddressIndex];
 
 							// Send the test to the device for execution!
-							MessageEndpoint->Send(new FAutomationWorkerRunTests(ExecutionCount, AddressIndex, NextTest->GetCommand()), DeviceAddress);
+							MessageEndpoint->Send(new FAutomationWorkerRunTests(ExecutionCount, AddressIndex, NextTest->GetCommand(), bScreenshotsEnabled, bRequestFullScreenScreenshots), DeviceAddress);
 
 							// Add a test so we can check later if the device is still active
 							TestRunningArray.Add( FTestRunningInfo( DeviceAddress ) );
diff --git a/Engine/Source/Developer/AutomationController/Public/Interfaces/IAutomationControllerManager.h b/Engine/Source/Developer/AutomationController/Public/Interfaces/IAutomationControllerManager.h
--- a/Engine/Source/Developer/AutomationController/Public/Interfaces/IAutomationControllerManager.h
+++ b/Engine/Source/Developer/AutomationController/Public/Interfaces/IAutomationControllerManager.h
@@ -152,6 +152,26 @@ public:
 	 */
 	virtual int32 GetNumPasses() = 0;
 
+	/**
+	 * Returns if we are using full size screen shots
+	 */
+	virtual bool IsUsingFullSizeScreenshots() const = 0;
+
+	/**
+	 * Sets if we are using full size screen shots
+	 */
+	virtual void SetUsingFullSizeScreenshots(const  bool bNewValue ) = 0;
+
+	/**
+	 * Returns if screenshots are enabled
+	 */
+	virtual bool AreScreenshotsEnabled() const = 0;
+
+	/**
+	 * Sets if screenshots are enabled
+	 */
+	virtual void SetScreenshotsEnabled( const bool bNewValue ) = 0;
+
 	/**
 	 * Filters the visible list of tests.
 	 */
diff --git a/Engine/Source/Developer/AutomationWindow/Private/SAutomationWindow.cpp b/Engine/Source/Developer/AutomationWindow/Private/SAutomationWindow.cpp
--- a/Engine/Source/Developer/AutomationWindow/Private/SAutomationWindow.cpp
+++ b/Engine/Source/Developer/AutomationWindow/Private/SAutomationWindow.cpp
@@ -330,41 +330,41 @@ void SAutomationWindow::CreateCommands()
 
 	ActionList.MapAction( Commands.RefreshTests,
 		FExecuteAction::CreateRaw( this, &SAutomationWindow::ListTests ),
-		FCanExecuteAction::CreateRaw( this, &SAutomationWindow::IsNotRunningTests )
+		FCanExecuteAction::CreateRaw( this, &SAutomationWindow::IsAutomationControllerIdle )
 		);
 
 	ActionList.MapAction( Commands.FindWorkers,
 		FExecuteAction::CreateRaw( this, &SAutomationWindow::FindWorkers ),
-		FCanExecuteAction::CreateRaw( this, &SAutomationWindow::IsNotRunningTests )
+		FCanExecuteAction::CreateRaw( this, &SAutomationWindow::IsAutomationControllerIdle )
 		);
 
 	ActionList.MapAction( Commands.ErrorFilter,
 		FExecuteAction::CreateRaw( this, &SAutomationWindow::OnToggleErrorFilter ),
-		FCanExecuteAction::CreateRaw( this, &SAutomationWindow::IsNotRunningTests ),
+		FCanExecuteAction::CreateRaw( this, &SAutomationWindow::IsAutomationControllerIdle ),
 		FIsActionChecked::CreateRaw( this, &SAutomationWindow::IsErrorFilterOn )
 		);
 
 	ActionList.MapAction( Commands.WarningFilter,
 		FExecuteAction::CreateRaw( this, &SAutomationWindow::OnToggleWarningFilter ),
-		FCanExecuteAction::CreateRaw( this, &SAutomationWindow::IsNotRunningTests ),
+		FCanExecuteAction::CreateRaw( this, &SAutomationWindow::IsAutomationControllerIdle ),
 		FIsActionChecked::CreateRaw( this, &SAutomationWindow::IsWarningFilterOn )
 		);
 
 	ActionList.MapAction( Commands.SmokeTestFilter,
 		FExecuteAction::CreateRaw( this, &SAutomationWindow::OnToggleSmokeTestFilter ),
-		FCanExecuteAction::CreateRaw( this, &SAutomationWindow::IsNotRunningTests ),
+		FCanExecuteAction::CreateRaw( this, &SAutomationWindow::IsAutomationControllerIdle ),
 		FIsActionChecked::CreateRaw( this, &SAutomationWindow::IsSmokeTestFilterOn )
 		);
 
 	ActionList.MapAction( Commands.DeveloperDirectoryContent,
 		FExecuteAction::CreateRaw( this, &SAutomationWindow::OnToggleDeveloperDirectoryIncluded ),
-		FCanExecuteAction::CreateRaw( this, &SAutomationWindow::IsNotRunningTests ),
+		FCanExecuteAction::CreateRaw( this, &SAutomationWindow::IsAutomationControllerIdle ),
 		FIsActionChecked::CreateRaw( this, &SAutomationWindow::IsDeveloperDirectoryIncluded )
 		);
 
 	ActionList.MapAction( Commands.VisualCommandlet,
 		FExecuteAction::CreateRaw( this, &SAutomationWindow::OnToggleVisualCommandletFilter ),
-		FCanExecuteAction::CreateRaw( this, &SAutomationWindow::IsNotRunningTests ),
+		FCanExecuteAction::CreateRaw( this, &SAutomationWindow::IsAutomationControllerIdle ),
 		FIsActionChecked::CreateRaw( this, &SAutomationWindow::IsVisualCommandletFilterOn )
 		);
 }
@@ -478,106 +478,117 @@ TSharedRef< SWidget > SAutomationWindow::MakeAutomationWindowToolBar( const TSha
 			SAssignNew( AutomationSearchBox, SSearchBox )
 			.ToolTipText( LOCTEXT( "Search Tests", "Search Tests" ) )
 			.OnTextChanged( this, &SAutomationWindow::OnFilterTextChanged )
-			.IsEnabled( this, &SAutomationWindow::IsNotRunningTests )
+			.IsEnabled( this, &SAutomationWindow::IsAutomationControllerIdle )
 			.MinDesiredWidth(SearchWidth)
 		];
 
 	TSharedRef<SWidget> TestPresets = 
-		SNew( SHorizontalBox )
-		+SHorizontalBox::Slot()
-		.AutoWidth()
-		.Padding(2.0f)
+		SNew( SVerticalBox )
+		+SVerticalBox::Slot()
+		.FillHeight(0.75f)
 		.VAlign(VAlign_Bottom)
+		.HAlign(HAlign_Left)
 		[
-			SNew( STextBlock )
-			.Text( LOCTEXT("AutomationPresetLabel", "Preset:") )
+			SNew( SHorizontalBox )
+			+SHorizontalBox::Slot()
+			.AutoWidth()
+			.Padding(2.0f)
+			.VAlign(VAlign_Bottom)
+			[
+				SNew( STextBlock )
+				.Text( LOCTEXT("AutomationPresetLabel", "Preset:") )
+				.IsEnabled( this, &SAutomationWindow::IsAutomationControllerIdle )
+			]
 		]
-		+SHorizontalBox::Slot()
-		.AutoWidth()
-		.Padding(2.0f)
+		+SVerticalBox::Slot()
+		.AutoHeight()
+		.VAlign(VAlign_Bottom)
 		[
-			//Preset Combo / Text
-			SNew(SOverlay)
-			+SOverlay::Slot()
+			SNew( SHorizontalBox )
+			+SHorizontalBox::Slot()
+			.AutoWidth()
 			[
-				SNew(SHorizontalBox)
-				.Visibility(this,&SAutomationWindow::HandlePresetComboVisibility)
-				+SHorizontalBox::Slot()
-				.FillWidth(1.0)
+				//Preset Combo / Text
+				SNew(SOverlay)
+				+SOverlay::Slot()
 				[
-					SAssignNew( PresetComboBox, SComboBox< TSharedPtr<FAutomationTestPreset> > )
-					.OptionsSource( &TestPresetManager->GetAllPresets() )
-					.OnGenerateWidget( this, &SAutomationWindow::GeneratePresetComboItem )
-					.OnSelectionChanged( this, &SAutomationWindow::HandlePresetChanged )
+					SNew(SHorizontalBox)
+					.Visibility(this,&SAutomationWindow::HandlePresetComboVisibility)
+					+SHorizontalBox::Slot()
+					.FillWidth(1.0)
 					[
-						SNew( STextBlock )
-						.Text( this, &SAutomationWindow::GetPresetComboText )
+						SAssignNew( PresetComboBox, SComboBox< TSharedPtr<FAutomationTestPreset> > )
+						.OptionsSource( &TestPresetManager->GetAllPresets() )
+						.OnGenerateWidget( this, &SAutomationWindow::GeneratePresetComboItem )
+						.OnSelectionChanged( this, &SAutomationWindow::HandlePresetChanged )
+						.IsEnabled(this, &SAutomationWindow::IsAutomationControllerIdle)
+						[
+							SNew( STextBlock )
+							.Text( this, &SAutomationWindow::GetPresetComboText )
+						]
 					]
 				]
-			]
-			+SOverlay::Slot()
-			[
-				SNew(SHorizontalBox)
-				.Visibility(this,&SAutomationWindow::HandlePresetTextVisibility)
-				+SHorizontalBox::Slot()
-				.FillWidth(1.0)
+				+SOverlay::Slot()
 				[
-					SAssignNew(PresetTextBox, SEditableTextBox)
-					.OnTextCommitted(this, &SAutomationWindow::HandlePresetTextCommited)
+					SNew(SHorizontalBox)
+					.Visibility(this,&SAutomationWindow::HandlePresetTextVisibility)
+					+SHorizontalBox::Slot()
+					.FillWidth(1.0)
+					[
+						SAssignNew(PresetTextBox, SEditableTextBox)
+						.OnTextCommitted(this, &SAutomationWindow::HandlePresetTextCommited)
+						.IsEnabled(this, &SAutomationWindow::IsAutomationControllerIdle)
+					]
 				]
 			]
-		]
 
-		//New button
-		+SHorizontalBox::Slot()
-		.AutoWidth()
-		.Padding(2.0f)
-		[
-			SNew( SButton )
-			.ButtonStyle( FEditorStyle::Get(), "NoBorder" )
-			.OnClicked( this, &SAutomationWindow::HandleNewPresetClicked )
-			.ToolTipText( LOCTEXT("AutomationPresetNewButtonTooltip", "Create a new preset") )
-			.IsEnabled(this, &SAutomationWindow::IsAddButtonIsEnabled)
-			.Content()
+			//New button
+			+SHorizontalBox::Slot()
+			.AutoWidth()
 			[
-				SNew(SImage)
-				.Image(FEditorStyle::Get().GetBrush("AutomationWindow.PresetNew"))
+				SNew( SButton )
+				.ButtonStyle( FEditorStyle::Get(), "NoBorder" )
+				.OnClicked( this, &SAutomationWindow::HandleNewPresetClicked )
+				.ToolTipText( LOCTEXT("AutomationPresetNewButtonTooltip", "Create a new preset") )
+				.IsEnabled(this, &SAutomationWindow::IsAddButtonEnabled)
+				.Content()
+				[
+					SNew(SImage)
+					.Image(FEditorStyle::Get().GetBrush("AutomationWindow.PresetNew"))
+				]
 			]
-		]
 
-		//Save button
-		+SHorizontalBox::Slot()
-		.AutoWidth()
-		.Padding(2.0f)
-		[
-			SNew( SButton )
-			.ButtonStyle( FEditorStyle::Get(), "NoBorder" )
-			.OnClicked( this, &SAutomationWindow::HandleSavePresetClicked )
-			.ToolTipText( LOCTEXT("AutomationPresetSaveButtonTooltip", "Save the current test list") )
-			.IsEnabled(this, &SAutomationWindow::IsSaveButtonIsEnabled)
-			.Content()
+			//Save button
+			+SHorizontalBox::Slot()
+			.AutoWidth()
 			[
-				SNew(SImage)
-				.Image(FEditorStyle::Get().GetBrush("AutomationWindow.PresetSave"))
+				SNew( SButton )
+				.ButtonStyle( FEditorStyle::Get(), "NoBorder" )
+				.OnClicked( this, &SAutomationWindow::HandleSavePresetClicked )
+				.ToolTipText( LOCTEXT("AutomationPresetSaveButtonTooltip", "Save the current test list") )
+				.IsEnabled(this, &SAutomationWindow::IsSaveButtonEnabled)
+				.Content()
+				[
+					SNew(SImage)
+					.Image(FEditorStyle::Get().GetBrush("AutomationWindow.PresetSave"))
+				]
 			]
-		]
 
-		//Remove button
-		+ SHorizontalBox::Slot()
-		.AutoWidth()
-		.Padding(2.0f)
-		[
-			// remove button
-			SNew(SButton)
-			.ContentPadding(2)
-			.ButtonStyle( FEditorStyle::Get(), "NoBorder" )
-			.OnClicked( this, &SAutomationWindow::HandleRemovePresetClicked )
-			.ToolTipText(LOCTEXT("AutomationPresetRemoveButtonTooltip", "Remove the selected preset"))
-			.IsEnabled(this, &SAutomationWindow::IsRemoveButtonIsEnabled)
-			.Content()
+			//Remove button
+			+ SHorizontalBox::Slot()
+			.AutoWidth()
 			[
-				SNew(SImage)
-				.Image(FEditorStyle::Get().GetBrush("AutomationWindow.PresetRemove"))
+				// remove button
+				SNew(SButton)
+				.ButtonStyle( FEditorStyle::Get(), "NoBorder" )
+				.OnClicked( this, &SAutomationWindow::HandleRemovePresetClicked )
+				.ToolTipText(LOCTEXT("AutomationPresetRemoveButtonTooltip", "Remove the selected preset"))
+				.IsEnabled(this, &SAutomationWindow::IsRemoveButtonEnabled)
+				.Content()
+				[
+					SNew(SImage)
+					.Image(FEditorStyle::Get().GetBrush("AutomationWindow.PresetRemove"))
+				]
 			]
 		];
 
@@ -612,19 +623,19 @@ EVisibility SAutomationWindow::HandlePresetTextVisibility( ) const
 	return bAddingTestPreset ? EVisibility::Visible : EVisibility::Hidden;
 }
 
-bool SAutomationWindow::IsAddButtonIsEnabled() const
+bool SAutomationWindow::IsAddButtonEnabled() const
 {
-	return !bAddingTestPreset;
+	return !bAddingTestPreset && IsAutomationControllerIdle();
 }
 
-bool SAutomationWindow::IsSaveButtonIsEnabled() const
+bool SAutomationWindow::IsSaveButtonEnabled() const
 {
-	return (!bAddingTestPreset && SelectedPreset.IsValid());
+	return (!bAddingTestPreset && SelectedPreset.IsValid() && IsAutomationControllerIdle());
 }
 
-bool SAutomationWindow::IsRemoveButtonIsEnabled() const
+bool SAutomationWindow::IsRemoveButtonEnabled() const
 {
-	return (!bAddingTestPreset && SelectedPreset.IsValid());
+	return (!bAddingTestPreset && SelectedPreset.IsValid() && IsAutomationControllerIdle());
 }
 
 void SAutomationWindow::HandlePresetTextCommited( const FText& CommittedText, ETextCommit::Type CommitType )
@@ -664,6 +675,28 @@ void SAutomationWindow::HandlePresetChanged( TSharedPtr<FAutomationTestPreset> I
 		SelectedPreset = Item;
 		AutomationController->SetEnabledTests(Item->GetEnabledTests());
 		TestTable->RequestTreeRefresh();
+		
+		//Expand selected items
+		TestTable->ClearExpandedItems();
+		TArray< TSharedPtr< IAutomationReport > >& TestReports = AutomationController->GetReports();
+		for( int32 Index = 0; Index < TestReports.Num(); Index++ )
+		{
+			ExpandEnabledTests(TestReports[Index]);
+		}
+	}
+}
+
+void SAutomationWindow::ExpandEnabledTests( TSharedPtr< IAutomationReport > InReport )
+{
+	// Expand node if the report is enabled or contains an enabled test
+	TestTable->SetItemExpansion( InReport, InReport->IsEnabled() || InReport->GetEnabledTestsNum() > 0 );
+
+	// Iterate through the child nodes to see if they should be expanded
+	TArray<TSharedPtr< IAutomationReport > > Reports = InReport->GetFilteredChildren();
+
+	for ( int32 ChildItem = 0; ChildItem < Reports.Num(); ChildItem++ )
+	{
+		ExpandEnabledTests( Reports[ ChildItem ] );
 	}
 }
 
@@ -733,27 +766,101 @@ TSharedRef< SWidget > SAutomationWindow::GenerateTestsOptionsMenuContent( )
 {
 	const bool bShouldCloseWindowAfterMenuSelection = true;
 	FMenuBuilder MenuBuilder( bShouldCloseWindowAfterMenuSelection, AutomationWindowActions );
-	TSharedRef<SWidget> NumTests = SNew(SSpinBox<int32>)
-			.MinValue(1)
-			.MaxValue(10)
-			.MinSliderValue(1)
-			.MaxSliderValue(10)
-			.Value(this,&SAutomationWindow::GetRepeatCount)
-			.OnValueChanged(this,&SAutomationWindow::OnChangeRepeatCount);
+	TSharedRef<SWidget> NumTests = 
+		SNew(SBox)
+		.WidthOverride( 200.0f )
+		[
+			SNew( SHorizontalBox )
+			+SHorizontalBox::Slot()
+			.Padding(0.0f,0.0f,4.0f, 0.0f)
+			.AutoWidth()
+			[
+				SNew( STextBlock )
+				.Text( LOCTEXT("NumTestsToolTip", "Number of runs:") )
+			]
+			+SHorizontalBox::Slot()
+			.FillWidth(1.f)
+			[
+				SNew(SSpinBox<int32>)
+				.MinValue(1)
+				.MaxValue(10)
+				.MinSliderValue(1)
+				.MaxSliderValue(10)
+				.Value(this,&SAutomationWindow::GetRepeatCount)
+				.OnValueChanged(this,&SAutomationWindow::OnChangeRepeatCount)
+				.IsEnabled( this, &SAutomationWindow::IsAutomationControllerIdle )
+			]
+
+		];
+
+	TSharedRef<SWidget> EnableScreenshotsWidget =
+		SNew(SCheckBox)
+		.IsChecked(this, &SAutomationWindow::IsEnableScreenshotsCheckBoxChecked)
+		.OnCheckStateChanged(this, &SAutomationWindow::HandleEnableScreenshotsBoxCheckStateChanged)
+		.Padding(FMargin(4.0f, 0.0f))
+		.ToolTipText(LOCTEXT("AutomationFullSizeScreenshotTip", "If checked, tests that support it will save screenshots"))
+		.IsEnabled( this, &SAutomationWindow::IsAutomationControllerIdle )
+		.Content()
+		[
+			SNew(STextBlock)
+			.Text(LOCTEXT("AutomationNoScreenshotText", "Enable screenshots"))
+		];
+
+	TSharedRef<SWidget> FullSizeScreenshotsWidget =
+		SNew(SCheckBox)
+		.IsChecked(this, &SAutomationWindow::IsFullSizeScreenshotsCheckBoxChecked)
+		.OnCheckStateChanged(this, &SAutomationWindow::HandleFullSizeScreenshotsBoxCheckStateChanged)
+		.Padding(FMargin(4.0f, 0.0f))
+		.ToolTipText(LOCTEXT("AutomationFullSizeScreenshotTip", "If checked, test that take screenshots will send back full size images instead of thumbnails."))
+		.IsEnabled( this, &SAutomationWindow::IsFullSizeScreenshotsOptionEnabled )
+		.Content()
+		[
+			SNew(STextBlock)
+			.Text(LOCTEXT("AutomationFullSizeScreenshotText", "Full size screenshots"))
+		];
+
+
 
 	MenuBuilder.BeginSection("AutomationWindowRunTest", LOCTEXT("RunTestOptions", "Advanced Settings"));
 	{
-		MenuBuilder.AddWidget( 
-			SNew( STextBlock )
-				.ColorAndOpacity(FSlateColor::UseSubduedForeground())
-				.Text(LOCTEXT("NumTestsToolTip", "Number of runs:")),
-			FText::GetEmpty() );
 		MenuBuilder.AddWidget(NumTests, FText::GetEmpty());
 	}
+	MenuBuilder.EndSection();
+	MenuBuilder.BeginSection("AutomationWindowScreenshots", LOCTEXT("ScreenshotOptions", "Screenshot Settings"));
+	{
+		MenuBuilder.AddWidget(EnableScreenshotsWidget, FText::GetEmpty());
+		MenuBuilder.AddWidget(FullSizeScreenshotsWidget, FText::GetEmpty());
+	}
+	MenuBuilder.EndSection();
 
 	return MenuBuilder.MakeWidget();
 }
 
+ESlateCheckBoxState::Type SAutomationWindow::IsFullSizeScreenshotsCheckBoxChecked() const
+{
+	return AutomationController->IsUsingFullSizeScreenshots() ? ESlateCheckBoxState::Checked : ESlateCheckBoxState::Unchecked;
+}
+
+void SAutomationWindow::HandleFullSizeScreenshotsBoxCheckStateChanged(ESlateCheckBoxState::Type CheckBoxState)
+{
+	AutomationController->SetUsingFullSizeScreenshots(CheckBoxState == ESlateCheckBoxState::Checked);
+}
+
+ESlateCheckBoxState::Type SAutomationWindow::IsEnableScreenshotsCheckBoxChecked() const
+{
+	return AutomationController->AreScreenshotsEnabled() ? ESlateCheckBoxState::Checked : ESlateCheckBoxState::Unchecked;
+}
+
+void SAutomationWindow::HandleEnableScreenshotsBoxCheckStateChanged(ESlateCheckBoxState::Type CheckBoxState)
+{
+	AutomationController->SetScreenshotsEnabled(CheckBoxState == ESlateCheckBoxState::Checked);
+}
+
+bool SAutomationWindow::IsFullSizeScreenshotsOptionEnabled() const
+{
+	return AutomationController->AreScreenshotsEnabled() && IsAutomationControllerIdle();
+}
+
 // Only valid in the editor
 #if WITH_EDITOR
 TSharedPtr<SWidget> SAutomationWindow::HandleAutomationListContextMenuOpening()
@@ -1289,7 +1396,7 @@ void SAutomationWindow::HandleSessionManagerSelectionChanged( const ISessionInfo
 }
 
 
-bool SAutomationWindow::IsNotRunningTests() const
+bool SAutomationWindow::IsAutomationControllerIdle() const
 {
 	return AutomationControllerState != EAutomationControllerModuleState::Running;
 }
@@ -1471,7 +1578,7 @@ void SAutomationWindow::HandleItemCheckBoxCheckedStateChanged( TSharedPtr< IAuto
 
 bool SAutomationWindow::HandleItemCheckBoxIsEnabled( ) const
 {
-	return IsNotRunningTests();
+	return IsAutomationControllerIdle();
 }
 
 
diff --git a/Engine/Source/Developer/AutomationWindow/Private/SAutomationWindow.h b/Engine/Source/Developer/AutomationWindow/Private/SAutomationWindow.h
--- a/Engine/Source/Developer/AutomationWindow/Private/SAutomationWindow.h
+++ b/Engine/Source/Developer/AutomationWindow/Private/SAutomationWindow.h
@@ -80,7 +80,7 @@ public:
 	*
 	* @return true if the tests aren't running
 	*/
-	bool IsNotRunningTests() const;
+	bool IsAutomationControllerIdle() const;
 
 protected:
 
@@ -229,6 +229,19 @@ private:
 	bool IsErrorFilterOn() const;
 	/** Toggles filtering of tests based on error condition */
 	void OnToggleErrorFilter();
+
+	/** Returns if full size screen shots are enabled */
+	ESlateCheckBoxState::Type IsFullSizeScreenshotsCheckBoxChecked() const;
+	/** Toggles if we are collecting full size screenshots */
+	void HandleFullSizeScreenshotsBoxCheckStateChanged(ESlateCheckBoxState::Type CheckBoxState);
+
+	/** Returns if screen shots are enabled */
+	ESlateCheckBoxState::Type IsEnableScreenshotsCheckBoxChecked() const;
+	/** Toggles if we are taking screenshots */
+	void HandleEnableScreenshotsBoxCheckStateChanged(ESlateCheckBoxState::Type CheckBoxState);
+
+	/** Returns if the full size screenshots option is enabled */
+	bool IsFullSizeScreenshotsOptionEnabled() const;
 	
 	/** Sets the number of times to repeat the tests */
 	void OnChangeRepeatCount(int32 InNewValue);
@@ -348,17 +361,17 @@ private:
 	/**
 	* Should the add preset button be enabled
 	*/
-	bool IsAddButtonIsEnabled() const;
+	bool IsAddButtonEnabled() const;
 
 	/**
 	* Should the save preset button be enabled
 	*/
-	bool IsSaveButtonIsEnabled() const;
+	bool IsSaveButtonEnabled() const;
 
 	/**
 	* Should the remove preset button be enabled
 	*/
-	bool IsRemoveButtonIsEnabled() const;
+	bool IsRemoveButtonEnabled() const;
 
 	/**
 	* Handles if the preset combo box should be visible
@@ -380,6 +393,11 @@ private:
 	*/
 	void HandlePresetChanged( TSharedPtr<FAutomationTestPreset> Item, ESelectInfo::Type SelectInfo );
 
+	/**
+	* Expands the test tree to show all enabled tests
+	*/
+	void ExpandEnabledTests( TSharedPtr< IAutomationReport > InReport );
+
 	/**
 	* Gets the text to display for the preset combo box
 	*/
diff --git a/Engine/Source/Developer/ScreenShotComparison/Private/ScreenShotComparisonPrivatePCH.h b/Engine/Source/Developer/ScreenShotComparison/Private/ScreenShotComparisonPrivatePCH.h
--- a/Engine/Source/Developer/ScreenShotComparison/Private/ScreenShotComparisonPrivatePCH.h
+++ b/Engine/Source/Developer/ScreenShotComparison/Private/ScreenShotComparisonPrivatePCH.h
@@ -13,6 +13,7 @@
 /* Dependencies
  *****************************************************************************/
 
+#include "ImageWrapper.h"
 #include "ModuleManager.h"
 #include "IFilter.h"
 #include "TextFilter.h"
@@ -30,5 +31,6 @@
 #include "SScreenPlatformRow.h"
 #include "SScreenViewRow.h"
 #include "SScreenShotBrowser.h"
+#include "SScreenShotImagePopup.h"
 
 #endif // __ScreenShotComparisonPrivatePCH_h__
diff --git a/Engine/Source/Developer/ScreenShotComparison/Private/Widgets/SScreenShotImagePopup.cpp b/Engine/Source/Developer/ScreenShotComparison/Private/Widgets/SScreenShotImagePopup.cpp
new file mode 100644
--- /dev/null
+++ b/Engine/Source/Developer/ScreenShotComparison/Private/Widgets/SScreenShotImagePopup.cpp
@@ -0,0 +1,30 @@
+// Copyright 1998-2014 Epic Games, Inc. All Rights Reserved.
+
+/*=============================================================================
+	SScreenShotImagePopup.cpp: Implements the SScreenImagePopup class.
+=============================================================================*/
+
+#include "ScreenShotComparisonPrivatePCH.h"
+
+void SScreenShotImagePopup::Construct( const FArguments& InArgs )
+{
+	// Load the image
+	DynamicImageBrush = MakeShareable(new FSlateDynamicImageBrush( InArgs._ImageAssetName, InArgs._ImageSize ) );
+
+	// Create the screen shot popup widget.
+	ChildSlot
+	[
+		SNew(SBorder)
+		. BorderImage(FEditorStyle::GetBrush("Menu.Background"))
+		. Padding(10)
+		[
+			SNew(SBox)
+			.WidthOverride(InArgs._ImageSize.X)
+			.HeightOverride(InArgs._ImageSize.Y)
+			[
+				SNew(SImage)
+				.Image( DynamicImageBrush.Get() )
+			]
+		]
+	];
+}
\ No newline at end of file
diff --git a/Engine/Source/Developer/ScreenShotComparison/Private/Widgets/SScreenShotImagePopup.h b/Engine/Source/Developer/ScreenShotComparison/Private/Widgets/SScreenShotImagePopup.h
new file mode 100644
--- /dev/null
+++ b/Engine/Source/Developer/ScreenShotComparison/Private/Widgets/SScreenShotImagePopup.h
@@ -0,0 +1,39 @@
+// Copyright 1998-2014 Epic Games, Inc. All Rights Reserved.
+
+/*=============================================================================
+	SScreenImagePopup.h: Declares the SScreenShotImagePopup class.
+=============================================================================*/
+
+#pragma once
+
+/**
+ * A widget to show a screen shot image as a pop up window
+ */
+class SScreenShotImagePopup : public SCompoundWidget
+{
+public:
+
+	SLATE_BEGIN_ARGS( SScreenShotImagePopup ) {}
+	
+	/** The assest name for the screen shot */
+	SLATE_ARGUMENT( FName, ImageAssetName )
+
+	/** The size of the image */
+	SLATE_ARGUMENT( FIntPoint, ImageSize )
+
+	SLATE_END_ARGS()
+
+	/**
+	 * Construct the widget.
+	 *
+	 * @param InArgs   A declaration from which to construct the widget.
+ 	 * @param InOwnerTableView   The owning table data.
+	 */
+	void Construct( const FArguments& InArgs );
+
+private:
+
+	// Holds the brush for this screenshot
+	TSharedPtr<FSlateDynamicImageBrush> DynamicImageBrush;
+};
+
diff --git a/Engine/Source/Developer/ScreenShotComparison/Private/Widgets/SScreenShotItem.cpp b/Engine/Source/Developer/ScreenShotComparison/Private/Widgets/SScreenShotItem.cpp
--- a/Engine/Source/Developer/ScreenShotComparison/Private/Widgets/SScreenShotItem.cpp
+++ b/Engine/Source/Developer/ScreenShotComparison/Private/Widgets/SScreenShotItem.cpp
@@ -10,6 +10,8 @@ void SScreenShotItem::Construct( const FArguments& InArgs )
 {
 	ScreenShotData = InArgs._ScreenShotData;
 
+	CachedActualImageSize = FIntPoint::NoneValue;
+
 	LoadBrush();
 	// Create the screen shot data widget.
 	ChildSlot
@@ -40,12 +42,74 @@ void SScreenShotItem::Construct( const FArguments& InArgs )
 					// The image to use. This needs to be a dynamic brush at some point.
 					SNew( SImage )
 					.Image( DynamicBrush.Get() )
+					.OnMouseButtonDown(this,&SScreenShotItem::OnImageClicked)
 				]
 			]
 		]
 	]; 
 }
 
+FIntPoint SScreenShotItem::GetActualImageSize()
+{
+	if( CachedActualImageSize == FIntPoint::NoneValue )
+	{
+		//Find the image size by reading in the image on disk.  
+		TArray<uint8> RawFileData;
+		if( FFileHelper::LoadFileToArray( RawFileData, *DynamicBrush->GetResourceName().ToString() ) )
+		{
+			IImageWrapperModule& ImageWrapperModule = FModuleManager::LoadModuleChecked<IImageWrapperModule>( FName("ImageWrapper") );
+			IImageWrapperPtr ImageWrapper = ImageWrapperModule.CreateImageWrapper( EImageFormat::PNG );
+
+			if ( ImageWrapper.IsValid() && ImageWrapper->SetCompressed( RawFileData.GetData(), RawFileData.Num() ) )
+			{			
+				const TArray<uint8>* RawData = NULL;
+				if (ImageWrapper->GetRaw( ERGBFormat::BGRA, 8, RawData))
+				{
+					CachedActualImageSize.X = ImageWrapper->GetWidth();
+					CachedActualImageSize.Y = ImageWrapper->GetHeight();
+				}
+			}
+		}
+	}
+
+	return CachedActualImageSize;
+}
+
+FReply SScreenShotItem::OnImageClicked(const FGeometry& InGeometry, const FPointerEvent& InEvent)
+{
+	const FIntPoint ImageSize = GetActualImageSize();
+
+	if( ImageSize.X > 256 || ImageSize.Y > 128 )
+	{
+		FString AssetFilename = ScreenShotData->GetAssetName();
+
+		TSharedRef<SScreenShotImagePopup> PopupImage = 
+			SNew(SScreenShotImagePopup)
+			. ImageAssetName(FName(*AssetFilename))
+			. ImageSize(ImageSize);
+
+		auto ParentWindow = FSlateApplication::Get().FindWidgetWindow(AsShared()).ToSharedRef();
+
+		// Center ourselves in the parent window
+		auto PopupWindow = SNew(SWindow)
+			.Title( FText::FromString(AssetFilename))
+			.IsPopupWindow(false)
+			.SizingRule(ESizingRule::Autosized)
+			.AutoCenter(EAutoCenter::PreferredWorkArea)
+			.SupportsMaximize(false)
+			.SupportsMinimize(false)
+			.FocusWhenFirstShown(true)
+			.ActivateWhenFirstShown(true)
+			.Content()
+			[
+				PopupImage
+			];
+
+		FSlateApplication::Get().AddWindowAsNativeChild(PopupWindow, ParentWindow, true );
+	}
+	return FReply::Handled();
+}
+
 void SScreenShotItem::LoadBrush()
 {
 	// The asset filename
diff --git a/Engine/Source/Developer/ScreenShotComparison/Private/Widgets/SScreenShotItem.h b/Engine/Source/Developer/ScreenShotComparison/Private/Widgets/SScreenShotItem.h
--- a/Engine/Source/Developer/ScreenShotComparison/Private/Widgets/SScreenShotItem.h
+++ b/Engine/Source/Developer/ScreenShotComparison/Private/Widgets/SScreenShotItem.h
@@ -37,8 +37,20 @@ private:
 	 */
 	void LoadBrush();
 
+	/**
+	 * Handles the user clicking on the image
+	 */
+	FReply OnImageClicked(const FGeometry& InGeometry, const FPointerEvent& InEvent);
+
+	/** 
+	 * Gets the actual size of the image
+	 */
+	FIntPoint GetActualImageSize();
 private:
 
+	//The cached actual size of the screenshot
+	FIntPoint CachedActualImageSize;
+
 	//Holds the dynamic brush.
 	TSharedPtr<FSlateDynamicImageBrush> DynamicBrush;
 
diff --git a/Engine/Source/Developer/ScreenShotComparison/Private/Widgets/SScreenShotSearchBar.cpp b/Engine/Source/Developer/ScreenShotComparison/Private/Widgets/SScreenShotSearchBar.cpp
--- a/Engine/Source/Developer/ScreenShotComparison/Private/Widgets/SScreenShotSearchBar.cpp
+++ b/Engine/Source/Developer/ScreenShotComparison/Private/Widgets/SScreenShotSearchBar.cpp
@@ -22,6 +22,9 @@ void SScreenShotSearchBar::Construct( const FArguments& InArgs, IScreenShotManag
 
 	PlatformDisplayString = "Any";
 
+	//Default to show all screenshots
+	DisplayEveryNthScreenshot = 1;
+
 	TArray< TSharedPtr<FString> >& PlatformList = ScreenShotManager->GetCachedPlatfomList();
 
 	ChildSlot
@@ -32,26 +35,85 @@ void SScreenShotSearchBar::Construct( const FArguments& InArgs, IScreenShotManag
 		+ SHorizontalBox::Slot()
 		.AutoWidth()
 		[
-			SNew( SBox )
-			.WidthOverride( 100.0f )
+			SNew( SHorizontalBox )
+			+SHorizontalBox::Slot()
+			.HAlign(HAlign_Center)
+			.VAlign(VAlign_Center)
+			.AutoWidth()
+			[
+				SNew( STextBlock )
+				.Text( LOCTEXT("ScreenshotNthLabel", "Display Every Nth Screenshot:") )
+			]
+			+SHorizontalBox::Slot()
+			.AutoWidth()
 			[
-				SNew(SComboButton)
-				.ButtonContent()
+				SNew( SBox )
+				.Padding( FMargin(4.0f, 0.0f, 0.0f, 0.0f) )
+				.WidthOverride( 100.0f )
 				[
-					SNew(STextBlock)
-					.Text( this, &SScreenShotSearchBar::GetPlatformString )
+					SNew(SSpinBox<int32>)
+					.MinValue(1)
+					.MaxValue(20)
+					.MinSliderValue(1)
+					.MaxSliderValue(20)
+					.Value(this,&SScreenShotSearchBar::GetDisplayEveryNthScreenshot)
+					.OnValueChanged(this,&SScreenShotSearchBar::OnChangedDisplayEveryNthScreenshot)
+					.OnValueCommitted(this,&SScreenShotSearchBar::OnCommitedDisplayEveryNthScreenshot)
 				]
-				.ContentPadding(FMargin(6.0f, 2.0f))
-				.MenuContent()
+			]
+		]
+		+SHorizontalBox::Slot()
+		.AutoWidth()
+		.Padding( 5.0f, 0.0f, 5.0f, 0.0f )
+		[
+			SNew(SSeparator)
+			.Orientation(Orient_Vertical)
+			.SeparatorImage( FEditorStyle::GetBrush(TEXT("Menu.Separator")) )
+		]
+		+ SHorizontalBox::Slot()
+		.AutoWidth()
+		[
+			SNew( SHorizontalBox )
+			+SHorizontalBox::Slot()
+			.HAlign(HAlign_Center)
+			.VAlign(VAlign_Center)
+			.AutoWidth()
+			[
+				SNew( STextBlock )
+				.Text( LOCTEXT("ScreenshotPlatform", "Platform:") )
+			]
+			+SHorizontalBox::Slot()
+			.AutoWidth()
+			[
+				SNew( SBox )
+				.WidthOverride( 100.0f )
 				[
-					SAssignNew(PlatformListView, SListView<TSharedPtr<FString> >)
-					.ItemHeight(24.0f)
-					.ListItemsSource(&PlatformList)
-					.OnGenerateRow(this, &SScreenShotSearchBar::HandlePlatformListViewGenerateRow)
-					.OnSelectionChanged(this, &SScreenShotSearchBar::HandlePlatformListSelected )
+					SNew(SComboButton)
+					.ButtonContent()
+					[
+						SNew(STextBlock)
+						.Text( this, &SScreenShotSearchBar::GetPlatformString )
+					]
+					.ContentPadding(FMargin(6.0f, 2.0f))
+						.MenuContent()
+						[
+							SAssignNew(PlatformListView, SListView<TSharedPtr<FString> >)
+							.ItemHeight(24.0f)
+							.ListItemsSource(&PlatformList)
+							.OnGenerateRow(this, &SScreenShotSearchBar::HandlePlatformListViewGenerateRow)
+							.OnSelectionChanged(this, &SScreenShotSearchBar::HandlePlatformListSelected )
+						]
 				]
 			]
 		]
+		+SHorizontalBox::Slot()
+		.AutoWidth()
+		.Padding( 5.0f, 0.0f, 5.0f, 0.0f )
+		[
+			SNew(SSeparator)
+			.Orientation(Orient_Vertical)
+			.SeparatorImage( FEditorStyle::GetBrush(TEXT("Menu.Separator")) )
+		]
 		+ SHorizontalBox::Slot()
 		.FillWidth( 0.3f )
 		[
@@ -73,6 +135,22 @@ void SScreenShotSearchBar::Construct( const FArguments& InArgs, IScreenShotManag
 	];
 }
 
+void SScreenShotSearchBar::OnCommitedDisplayEveryNthScreenshot(int32 InNewValue, ETextCommit::Type CommitType)
+{
+	DisplayEveryNthScreenshot = InNewValue;
+	ScreenShotManager->SetDisplayEveryNthScreenshot(DisplayEveryNthScreenshot);
+}
+
+void SScreenShotSearchBar::OnChangedDisplayEveryNthScreenshot(int32 InNewValue)
+{
+	DisplayEveryNthScreenshot = InNewValue;
+}
+
+int32 SScreenShotSearchBar::GetDisplayEveryNthScreenshot() const
+{
+	return DisplayEveryNthScreenshot;
+}
+
 FString SScreenShotSearchBar::GetPlatformString() const
 {
 	return PlatformDisplayString;
@@ -90,7 +168,18 @@ void SScreenShotSearchBar::HandleFilterTextChanged( const FText& InFilterText )
 
 void SScreenShotSearchBar::HandlePlatformListSelected( TSharedPtr<FString> PlatformName, ESelectInfo::Type SelectInfo )
 {
-	FString TestName = *PlatformName;
+	FString TestName;
+
+	if( PlatformName.IsValid() )
+	{
+		TestName = *PlatformName;
+		PlatformDisplayString = *PlatformName;
+	}
+	else
+	{
+		PlatformDisplayString = "Any";
+	}
+
 	if ( TestName == "Any" )
 	{
 		TestName = "";
@@ -98,7 +187,6 @@ void SScreenShotSearchBar::HandlePlatformListSelected( TSharedPtr<FString> Platf
 
 	ScreenShotComparisonFilter->SetPlatformFilter( TestName );
 	ScreenShotManager->SetFilter( ScreenShotFilters );
-	PlatformDisplayString = *PlatformName;
 }
 
 
@@ -117,7 +205,8 @@ FReply SScreenShotSearchBar::RefreshView( )
 
 	ScreenShotComparisonFilter->SetPlatformFilter( "" );
 	ScreenShotManager->SetFilter( ScreenShotFilters );
-	PlatformDisplayString = "";
+	PlatformDisplayString = "Any";
+	DisplayEveryNthScreenshot = 1;
 
 	return FReply::Handled();
 }
diff --git a/Engine/Source/Developer/ScreenShotComparison/Private/Widgets/SScreenShotSearchBar.h b/Engine/Source/Developer/ScreenShotComparison/Private/Widgets/SScreenShotSearchBar.h
--- a/Engine/Source/Developer/ScreenShotComparison/Private/Widgets/SScreenShotSearchBar.h
+++ b/Engine/Source/Developer/ScreenShotComparison/Private/Widgets/SScreenShotSearchBar.h
@@ -62,7 +62,24 @@ protected:
 	 */
 	TSharedRef<ITableRow> HandlePlatformListViewGenerateRow( TSharedPtr<FString> PlatformName, const TSharedRef<STableViewBase>& OwnerTable );
 
+	/**
+	 * Sets the new N value for displaying every Nth screenshot.
+	 */
+	void OnCommitedDisplayEveryNthScreenshot(int32 InNewValue, ETextCommit::Type CommitType);
+
+	/**
+	 * Updates the display for the display every Nth screenshot widget.
+	 */
+	void OnChangedDisplayEveryNthScreenshot(int32 InNewValue);
+
+	/**
+	 * Gets the value of DisplayEveryNthScreenshot.
+	 */
+	int32 GetDisplayEveryNthScreenshot() const;
+
 private:
+	// Will only show every Nth screenshot
+	int32 DisplayEveryNthScreenshot;
 
 	// Holds the display string for the selected platform
 	FString PlatformDisplayString;
diff --git a/Engine/Source/Developer/ScreenShotComparison/ScreenShotComparison.Build.cs b/Engine/Source/Developer/ScreenShotComparison/ScreenShotComparison.Build.cs
--- a/Engine/Source/Developer/ScreenShotComparison/ScreenShotComparison.Build.cs
+++ b/Engine/Source/Developer/ScreenShotComparison/ScreenShotComparison.Build.cs
@@ -19,6 +19,7 @@ public class ScreenShotComparison : ModuleRules
 				"ScreenShotComparisonTools",
 				"Slate",
 				"SlateCore",
+				"ImageWrapper",
 			}
 		);
 
diff --git a/Engine/Source/Developer/ScreenShotComparisonTools/Private/ScreenShotBaseNode.cpp b/Engine/Source/Developer/ScreenShotComparisonTools/Private/ScreenShotBaseNode.cpp
--- a/Engine/Source/Developer/ScreenShotComparisonTools/Private/ScreenShotBaseNode.cpp
+++ b/Engine/Source/Developer/ScreenShotComparisonTools/Private/ScreenShotBaseNode.cpp
@@ -49,6 +49,14 @@ EScreenShotDataType::Type FScreenShotBaseNode::GetScreenNodeType()
 	return EScreenShotDataType::SSDT_Base; 
 };
 
+void FScreenShotBaseNode::SetDisplayEveryNthScreenshot( int32 NewDisplayEveryNth )
+{
+	for ( int32 Index = 0; Index < Children.Num(); Index++ )
+	{
+		Children[Index]->SetDisplayEveryNthScreenshot(NewDisplayEveryNth);
+	}
+}
+
 
 bool FScreenShotBaseNode::SetFilter( TSharedPtr< ScreenShotFilterCollection > ScreenFilter )
 {
diff --git a/Engine/Source/Developer/ScreenShotComparisonTools/Private/ScreenShotBaseNode.h b/Engine/Source/Developer/ScreenShotComparisonTools/Private/ScreenShotBaseNode.h
--- a/Engine/Source/Developer/ScreenShotComparisonTools/Private/ScreenShotBaseNode.h
+++ b/Engine/Source/Developer/ScreenShotComparisonTools/Private/ScreenShotBaseNode.h
@@ -37,6 +37,7 @@ public:
 	virtual const FString& GetName() const OVERRIDE;
 	virtual EScreenShotDataType::Type GetScreenNodeType() OVERRIDE;
 	virtual bool SetFilter( TSharedPtr< ScreenShotFilterCollection > ScreenFilter ) OVERRIDE;
+	virtual void SetDisplayEveryNthScreenshot( int32 NewNthValue ) OVERRIDE;
 
 	// End IScreenShotData interface
 
diff --git a/Engine/Source/Developer/ScreenShotComparisonTools/Private/ScreenShotManager.cpp b/Engine/Source/Developer/ScreenShotComparisonTools/Private/ScreenShotManager.cpp
--- a/Engine/Source/Developer/ScreenShotComparisonTools/Private/ScreenShotManager.cpp
+++ b/Engine/Source/Developer/ScreenShotComparisonTools/Private/ScreenShotManager.cpp
@@ -33,7 +33,6 @@ void FScreenShotManager::CreateData()
 
 	// Root directory
 	FString PathName = FPaths::AutomationDir();
-	PathName += "AutomationTest/";
 
 	// Array to store the list of platforms
 	TArray<FString> TempPlatforms;
@@ -47,6 +46,8 @@ void FScreenShotManager::CreateData()
 		FString ChangeListName;
 		FString PathRemainder;
 
+		FPaths::MakePathRelativeTo(PathString, *PathName);
+
 		if (PathString.Split(TEXT("/"), &PathRemainder, &ChangeListName, ESearchCase::CaseSensitive, ESearchDir::FromEnd) )
 		{
 			PathString = PathRemainder;
@@ -61,8 +62,9 @@ void FScreenShotManager::CreateData()
 				PathString = PathRemainder;
 				TempPlatforms.AddUnique( PlatformName );
 
-				FString ScreenShotName;
-				if ( PathString.Split(TEXT("/"), &PathRemainder, &ScreenShotName, ESearchCase::CaseSensitive, ESearchDir::FromEnd) )
+				//Take the rest of the path as the screenshot name
+				FString ScreenShotName = PathString;
+				if ( !ScreenShotName.IsEmpty() )
 				{
  					FScreenShotDataItem ScreenItem( ScreenShotName, PlatformName, AllPNGFiles[PathIndex], ChangeListNumber );
  					ScreenShotDataArray.Add( ScreenItem );
@@ -148,6 +150,14 @@ void FScreenShotManager::SetFilter( TSharedPtr< ScreenShotFilterCollection > InF
 	ScreenFilterChangedDelegate.ExecuteIfBound();
 }
 
+void FScreenShotManager::SetDisplayEveryNthScreenshot(int32 NewNth)
+{
+	ScreenShotRoot->SetDisplayEveryNthScreenshot(NewNth);
+
+	// Update the UI
+	ScreenFilterChangedDelegate.ExecuteIfBound();
+}
+
 
 /* IScreenShotManager event handlers
  *****************************************************************************/
@@ -155,7 +165,7 @@ void FScreenShotManager::SetFilter( TSharedPtr< ScreenShotFilterCollection > InF
 void FScreenShotManager::HandleScreenShotMessage( const FAutomationWorkerScreenImage& Message, const IMessageContextRef& Context )
 {
 	bool bTree = true;
-	FString FileName = FPaths::AutomationDir() + Message.ScreenShotName;
+	FString FileName = FPaths::RootDir() + Message.ScreenShotName;
 	IFileManager::Get().MakeDirectory( *FPaths::GetPath(FileName), bTree );
 	FFileHelper::SaveArrayToFile( Message.ScreenImage, *FileName );
 
diff --git a/Engine/Source/Developer/ScreenShotComparisonTools/Private/ScreenShotManager.h b/Engine/Source/Developer/ScreenShotComparisonTools/Private/ScreenShotManager.h
--- a/Engine/Source/Developer/ScreenShotComparisonTools/Private/ScreenShotManager.h
+++ b/Engine/Source/Developer/ScreenShotComparisonTools/Private/ScreenShotManager.h
@@ -43,6 +43,8 @@ public:
 
 	virtual void SetFilter( TSharedPtr< ScreenShotFilterCollection > InFilter ) OVERRIDE;
 
+	virtual void SetDisplayEveryNthScreenshot(int32 NewNth) OVERRIDE;
+
 	// End IScreenShotManager interface
 
 private:
diff --git a/Engine/Source/Developer/ScreenShotComparisonTools/Private/ScreenShotPlatformNode.h b/Engine/Source/Developer/ScreenShotComparisonTools/Private/ScreenShotPlatformNode.h
--- a/Engine/Source/Developer/ScreenShotComparisonTools/Private/ScreenShotPlatformNode.h
+++ b/Engine/Source/Developer/ScreenShotComparisonTools/Private/ScreenShotPlatformNode.h
@@ -17,7 +17,8 @@ public:
 	* @param InName  The Name of this asset.
 	*/
 	FScreenShotPlatformNode( const FString& InName, const FString& InAssetName = "" )
-	: FScreenShotBaseNode( InName, InAssetName )
+	: FScreenShotBaseNode( InName, InAssetName ),
+	DisplayEveryNthScreenshot(0)
 	{}
 
 public:
@@ -47,6 +48,10 @@ public:
 		{
 			for ( int32 Index = 0; Index < Children.Num(); Index++ )
 			{
+				if( DisplayEveryNthScreenshot > 1 && Index % DisplayEveryNthScreenshot != 0 )
+				{
+					continue;
+				}
 				FilteredChildren.Add( Children[Index] );
 			}
 		}
@@ -54,7 +59,28 @@ public:
 		return bPassesFilter;
 	};
 
+	virtual void SetDisplayEveryNthScreenshot( int32 NewLastNth ) OVERRIDE
+	{
+		DisplayEveryNthScreenshot = NewLastNth;
+		if( FilteredChildren.Num() > 0 )
+		{
+			FilteredChildren.Empty();
+			for ( int32 Index = 0; Index < Children.Num(); Index++ )
+			{
+				if( DisplayEveryNthScreenshot > 1 && Index % DisplayEveryNthScreenshot != 0 )
+				{
+					continue;
+				}
+				FilteredChildren.Add( Children[Index] );
+			}
+		}
+	}
+
 
 	// End IScreenShotData interface
+
+private:
+	// If > 1 we will only show every Nth screenshot.
+	uint32 DisplayEveryNthScreenshot;
 };
 
diff --git a/Engine/Source/Developer/ScreenShotComparisonTools/Public/Interfaces/IScreenShotData.h b/Engine/Source/Developer/ScreenShotComparisonTools/Public/Interfaces/IScreenShotData.h
--- a/Engine/Source/Developer/ScreenShotComparisonTools/Public/Interfaces/IScreenShotData.h
+++ b/Engine/Source/Developer/ScreenShotComparisonTools/Public/Interfaces/IScreenShotData.h
@@ -87,4 +87,11 @@ public:
 	* @return True if passes the filter
 	*/
 	virtual bool SetFilter( TSharedPtr< ScreenShotFilterCollection > ScreenFilter ) = 0;
+
+	/**
+	* Tells the node that we should only display every Nth screenshot
+	*
+	* @param NewDisplayEveryNth - The new N.
+	*/
+	virtual void SetDisplayEveryNthScreenshot( int32 NewDisplayEveryNth ) = 0;
 };
\ No newline at end of file
diff --git a/Engine/Source/Developer/ScreenShotComparisonTools/Public/Interfaces/IScreenShotManager.h b/Engine/Source/Developer/ScreenShotComparisonTools/Public/Interfaces/IScreenShotManager.h
--- a/Engine/Source/Developer/ScreenShotComparisonTools/Public/Interfaces/IScreenShotManager.h
+++ b/Engine/Source/Developer/ScreenShotComparisonTools/Public/Interfaces/IScreenShotManager.h
@@ -64,4 +64,11 @@ public:
 	* @param InFilter - The screen shot filter
 	*/
 	virtual void SetFilter(TSharedPtr<ScreenShotFilterCollection> InFilter ) = 0;
+
+	/**
+	* Only display every Nth screen shot
+	*
+	* @param NewNth - The new N.
+	*/
+	virtual void SetDisplayEveryNthScreenshot(int32 NewNth) = 0;
 };
diff --git a/Engine/Source/Editor/DetailCustomizations/Private/DirectoryPathStructCustomization.cpp b/Engine/Source/Editor/DetailCustomizations/Private/DirectoryPathStructCustomization.cpp
--- a/Engine/Source/Editor/DetailCustomizations/Private/DirectoryPathStructCustomization.cpp
+++ b/Engine/Source/Editor/DetailCustomizations/Private/DirectoryPathStructCustomization.cpp
@@ -15,6 +15,7 @@ TSharedRef<IStructCustomization> FDirectoryPathStructCustomization::MakeInstance
 void FDirectoryPathStructCustomization::CustomizeStructHeader( TSharedRef<IPropertyHandle> StructPropertyHandle, class FDetailWidgetRow& HeaderRow, IStructCustomizationUtils& StructCustomizationUtils )
 {
 	TSharedPtr<IPropertyHandle> PathProperty = StructPropertyHandle->GetChildHandle("Path");
+	bool bUseRelativePath = StructPropertyHandle->GetProperty()->HasMetaData( TEXT("RelativePath") );
 	if(PathProperty.IsValid())
 	{
 		HeaderRow.ValueContent()
@@ -36,7 +37,7 @@ void FDirectoryPathStructCustomization::CustomizeStructHeader( TSharedRef<IPrope
 				SAssignNew(BrowseButton, SButton)
 				.ButtonStyle( FEditorStyle::Get(), "HoverHintOnly" )
 				.ToolTipText( LOCTEXT( "FolderButtonToolTipText", "Choose a directory from this computer").ToString() )
-				.OnClicked( FOnClicked::CreateSP(this, &FDirectoryPathStructCustomization::OnPickDirectory, PathProperty.ToSharedRef()) )
+				.OnClicked( FOnClicked::CreateSP(this, &FDirectoryPathStructCustomization::OnPickDirectory, PathProperty.ToSharedRef(), bUseRelativePath) )
 				.ContentPadding( 2.0f )
 				.ForegroundColor( FSlateColor::UseForeground() )
 				.IsFocusable( false )
@@ -58,7 +59,7 @@ void FDirectoryPathStructCustomization::CustomizeStructChildren( TSharedRef<IPro
 {
 }
 
-FReply FDirectoryPathStructCustomization::OnPickDirectory(TSharedRef<IPropertyHandle> PropertyHandle) const
+FReply FDirectoryPathStructCustomization::OnPickDirectory(TSharedRef<IPropertyHandle> PropertyHandle, const bool bUseRelativePath) const
 {
 	FString Directory;
 	IDesktopPlatform* DesktopPlatform = FDesktopPlatformModule::Get();
@@ -69,8 +70,14 @@ FReply FDirectoryPathStructCustomization::OnPickDirectory(TSharedRef<IPropertyHa
 
 		if(DesktopPlatform->OpenDirectoryDialog(ParentWindowHandle, LOCTEXT( "FolderDialogTitle", "Choose a directory").ToString(), FEditorDirectories::Get().GetLastDirectory(ELastDirectory::GENERIC_IMPORT), Directory))
 		{
-			PropertyHandle->SetValue(Directory);
 			FEditorDirectories::Get().SetLastDirectory(ELastDirectory::GENERIC_IMPORT, Directory);
+
+			if( bUseRelativePath )
+			{
+				Directory = IFileManager::Get().ConvertToRelativePath(*Directory);
+			}
+
+			PropertyHandle->SetValue(Directory);
 		}
 	}
 
diff --git a/Engine/Source/Editor/DetailCustomizations/Private/DirectoryPathStructCustomization.h b/Engine/Source/Editor/DetailCustomizations/Private/DirectoryPathStructCustomization.h
--- a/Engine/Source/Editor/DetailCustomizations/Private/DirectoryPathStructCustomization.h
+++ b/Engine/Source/Editor/DetailCustomizations/Private/DirectoryPathStructCustomization.h
@@ -18,7 +18,7 @@ private:
 	FText GetDisplayedText(TSharedRef<IPropertyHandle> PropertyHandle) const;
 
 	/** Delegate used to display a directory picker */
-	FReply OnPickDirectory(TSharedRef<IPropertyHandle> PropertyHandle) const;
+	FReply OnPickDirectory(TSharedRef<IPropertyHandle> PropertyHandle, const bool bUseRelativePaths) const;
 
 	/** The browse button widget */
 	TSharedPtr<SButton> BrowseButton;
diff --git a/Engine/Source/Editor/StaticMeshEditor/Private/StaticMeshAutomationTests.cpp b/Engine/Source/Editor/StaticMeshEditor/Private/StaticMeshAutomationTests.cpp
--- a/Engine/Source/Editor/StaticMeshEditor/Private/StaticMeshAutomationTests.cpp
+++ b/Engine/Source/Editor/StaticMeshEditor/Private/StaticMeshAutomationTests.cpp
@@ -39,6 +39,37 @@ namespace EditorViewButtonHelper
 		EStaticMeshFlag::Type CommandType;
 	};
 
+	static const TCHAR* GetStaticMeshFlagName(EStaticMeshFlag::Type InType)
+	{
+		switch(InType)
+		{
+		case EStaticMeshFlag::Wireframe:
+			return TEXT("Wireframe");
+		case EStaticMeshFlag::Vert:
+			return TEXT("Vertex");
+		case EStaticMeshFlag::Grid:
+			return TEXT("Grid");
+		case EStaticMeshFlag::Bounds:
+			return TEXT("Bounds");
+		case EStaticMeshFlag::Collision:
+			return TEXT("Collision");
+		case EStaticMeshFlag::Pivot:
+			return TEXT("Pivot");
+		case EStaticMeshFlag::Normals:
+			return TEXT("Normals");
+		case EStaticMeshFlag::Tangents:
+			return TEXT("Tangents");
+		case EStaticMeshFlag::Binormals:
+			return TEXT("Binormals");
+		case EStaticMeshFlag::UV:
+			return TEXT("UV");
+		default:
+			break;
+		}
+
+		return TEXT("Unknown");
+	}
+
 	/**
 	 * Used for toggling a certain button on/off based on the passed in viewport client and button type.
 	 */
@@ -150,29 +181,54 @@ bool FStaticMeshEditorTest::RunTest(const FString& Parameters)
 	FAssetEditorManager::Get().OpenEditorForAsset(EditorMesh);
 	WindowParameters.CurrentWindow = FSlateApplication::Get().GetActiveTopLevelWindow();
 
+	//Check that we have the right window (Tutorial may have opened on top of the editor)
+	if( WindowParameters.CurrentWindow->GetTitle().ToString() != LoadedObjectType )
+	{
+		//Close the tutorial
+		WindowParameters.CurrentWindow->RequestDestroyWindow();
+
+		//Bring the asset editor to the front
+		FAssetEditorManager::Get().FindEditorForAsset(EditorMesh,false);
+		WindowParameters.CurrentWindow = FSlateApplication::Get().GetActiveTopLevelWindow();
+	}
+
 	//Grab the last opened Viewport (aka the asset manager)
 	AutomationParameters.ViewportClient = static_cast<FStaticMeshEditorViewportClient*>(GEditor->AllViewportClients[GEditor->AllViewportClients.Num() - 1]);
 	AutomationParameters.CommandType = EditorViewButtonHelper::EStaticMeshFlag::Wireframe; //Start it off at Wireframe because that is the first button we are attempting to press
 
 	if (AutomationParameters.ViewportClient)
 	{
-		const FString BaseFileName = FPaths::ScreenShotDir() / TEXT("StaticMeshEditor");
-		WindowParameters.ScreenshotName = FString::Printf(TEXT("%s00.bmp"), *BaseFileName);
+		const FString BaseFileName = TEXT("StaticMeshEditorTest");
+		FString TestName = FString::Printf(TEXT("%s/Base"), *BaseFileName);
 
 		//Wait for the window to load and then take the initial screen shot		
 		ADD_LATENT_AUTOMATION_COMMAND(FWaitLatentCommand(0.5f));
-		ADD_LATENT_AUTOMATION_COMMAND(FTakeEditorScreenshotCommand(WindowParameters));
-		ADD_LATENT_AUTOMATION_COMMAND(FWaitLatentCommand(0.1f));
+		if( FAutomationTestFramework::GetInstance().AreScreenshotsEnabled() )
+		{
+			AutomationCommon::GetScreenshotPath(TestName, WindowParameters.ScreenshotName, false);
+			ADD_LATENT_AUTOMATION_COMMAND(FTakeEditorScreenshotCommand(WindowParameters));
+			ADD_LATENT_AUTOMATION_COMMAND(FWaitLatentCommand(0.1f));
+		}
 
 		for(auto CurrentCommandType = 0; CurrentCommandType < EditorViewButtonHelper::EStaticMeshFlag::Max; ++CurrentCommandType)
 		{
-			WindowParameters.ScreenshotName = FString::Printf(TEXT("%s%02i.bmp"), *BaseFileName, CurrentCommandType + 1);
-
+			TestName = FString::Printf(TEXT("%s/%s"), *BaseFileName, EditorViewButtonHelper::GetStaticMeshFlagName(AutomationParameters.CommandType));
+			
 			//Toggle the command for the button, take a screenshot, and then re-toggle the command for the button
 			ADD_LATENT_AUTOMATION_COMMAND(EditorViewButtonHelper::FPerformStaticMeshFlagToggle(AutomationParameters));
 			ADD_LATENT_AUTOMATION_COMMAND(FWaitLatentCommand(0.1f));
-			ADD_LATENT_AUTOMATION_COMMAND(FTakeEditorScreenshotCommand(WindowParameters));
-			ADD_LATENT_AUTOMATION_COMMAND(FWaitLatentCommand(0.1f));
+
+			if( FAutomationTestFramework::GetInstance().AreScreenshotsEnabled() )
+			{
+				//Take the screenshot
+				AutomationCommon::GetScreenshotPath(TestName, WindowParameters.ScreenshotName, false);
+				ADD_LATENT_AUTOMATION_COMMAND(FTakeEditorScreenshotCommand(WindowParameters));
+
+				//Wait so the screenshots have a chance to save
+				ADD_LATENT_AUTOMATION_COMMAND(FWaitLatentCommand(0.1f));
+			}
+
+			//Toggle the flag back off
 			ADD_LATENT_AUTOMATION_COMMAND(EditorViewButtonHelper::FPerformStaticMeshFlagToggle(AutomationParameters));
 			
 			//Increment to the next enum, so we can loop through all the types
diff --git a/Engine/Source/Editor/UnrealEd/Classes/Settings/ProjectPackagingSettings.h b/Engine/Source/Editor/UnrealEd/Classes/Settings/ProjectPackagingSettings.h
--- a/Engine/Source/Editor/UnrealEd/Classes/Settings/ProjectPackagingSettings.h
+++ b/Engine/Source/Editor/UnrealEd/Classes/Settings/ProjectPackagingSettings.h
@@ -32,22 +32,6 @@ enum EProjectPackagingBuildConfigurations
 };
 
 
-/**
- * Structure for directory paths that are displayed in the UI.
- */
-USTRUCT()
-struct FDirectoryPath
-{
-	GENERATED_USTRUCT_BODY()
-
-	/**
-	 * The path to the directory.
-	 */
-	UPROPERTY(EditAnywhere, Category=Path)
-	FString Path;
-};
-
-
 /**
  * Implements the Editor's user settings.
  */
diff --git a/Engine/Source/Editor/UnrealEd/Private/Tests/EditorAutomationTests.cpp b/Engine/Source/Editor/UnrealEd/Private/Tests/EditorAutomationTests.cpp
--- a/Engine/Source/Editor/UnrealEd/Private/Tests/EditorAutomationTests.cpp
+++ b/Engine/Source/Editor/UnrealEd/Private/Tests/EditorAutomationTests.cpp
@@ -874,7 +874,45 @@ bool FLoadAllMapsInEditorTest::RunTest(const FString& Parameters)
 {
 	FString MapName = Parameters;
 
-	FEditorAutomationTestUtilities::LoadMap(MapName);
+	const bool bTakeScreenshots = FAutomationTestFramework::GetInstance().AreScreenshotsEnabled();
+	if( bTakeScreenshots )
+	{
+		//Find the main editor window
+		TArray<TSharedRef<SWindow> > AllWindows;
+		FSlateApplication::Get().GetAllVisibleWindowsOrdered(AllWindows);
+		if( AllWindows.Num() == 0 )
+		{
+			UE_LOG(LogEditorAutomationTests, Error, TEXT("ERROR: Could not find the main editor window."));
+			return false;
+		}
+		WindowScreenshotParameters WindowParameters;
+		WindowParameters.CurrentWindow = AllWindows[0];
+
+		//Disable Eye Adaptation
+		static IConsoleVariable* CVar = IConsoleManager::Get().FindConsoleVariable(TEXT("r.EyeAdaptationQuality"));
+		CVar->Set(0);
+
+		//Create a screen shot filename and path
+		const FString TestName = FString::Printf(TEXT("LoadAllMaps_Editor/%s"), *FPaths::GetBaseFilename(MapName));
+		AutomationCommon::GetScreenshotPath(TestName, WindowParameters.ScreenshotName, true);
+
+		//Load the map
+		FEditorAutomationTestUtilities::LoadMap(MapName);
+
+		//If we don't have NoTextureStreaming enabled, give the textures some time to load.
+		if( !FParse::Param( FCommandLine::Get(), TEXT( "NoTextureStreaming" ) ) )
+		{
+			//Give the contents some time to load
+			ADD_LATENT_AUTOMATION_COMMAND(FWaitLatentCommand(1.5f));
+		}
+		//Take the screen shot
+		ADD_LATENT_AUTOMATION_COMMAND(FTakeEditorScreenshotCommand(WindowParameters));
+	}
+	else
+	{
+		//Load the map
+		FEditorAutomationTestUtilities::LoadMap(MapName);
+	}
 
 	return true;
 }
@@ -1844,62 +1882,168 @@ bool FTraceAllTimelinesAutomationTest::RunTest(const FString& BlueprintName)
 	return bPassed;
 }
 
-/**
- * StaticMeshUVsTest
- */
-IMPLEMENT_SIMPLE_AUTOMATION_TEST( FStaticMeshUVsTest, "Editor.Content.Static Mesh UVs", EAutomationTestFlags::ATF_Editor )
-
-/**
- * Find all static meshes and check the lightmap UVs
- */
-bool FStaticMeshUVsTest::RunTest(const FString& Parameters)
+/** Static Mesh UV test specific data */
+namespace StaticMeshUVTest
 {
-	TArray<FString> MissingUVMessages;
-	TArray<FString> BadUVMessages;
-	TArray<FString> ValidUVMessages;
-	
-	//Create our package list
-	TArray<FString> ContentPackages;
-	FPackageName::FindPackagesInDirectory(ContentPackages, *FPaths::EngineContentDir());
-	FPackageName::FindPackagesInDirectory(ContentPackages, *FPaths::GameContentDir());
-	UE_LOG(LogEditorAutomationTests, Log, TEXT("Found %i Packages."), ContentPackages.Num() );
+	/** States the test can be in */
+	namespace EStaticMeshUVTestState
+	{
+		enum Type
+		{
+			Ready,
+			WaitingForPackage,
+			PackageLoaded
+		};
+	}
 
-	for( int32 PackageID = 0; PackageID < ContentPackages.Num(); ++PackageID)
+	class FUVTestHelper
 	{
-		const FString PackageName = ContentPackages[PackageID];
+	public:
+
+		/** Constructor */
+		FUVTestHelper(FAutomationTestBase* BaseTest) :
+			CurrentTest(BaseTest),
+			CurrentPackage(NULL),
+			LoadedPackageCount(0),
+			CurrentState(EStaticMeshUVTestState::Ready)
+		{
+			FPackageName::FindPackagesInDirectory(ContentPackages, *FPaths::EngineContentDir());
+			FPackageName::FindPackagesInDirectory(ContentPackages, *FPaths::GameContentDir());
+		}
 
-		UPackage* Package = NULL;
+		/**
+		 * Update the static mesh UV test
+		 * @return true if the test is complete
+		 */
+		bool Update()
 		{
-			SetSuppressLogs(true);	
-			Package = LoadPackage( NULL, *PackageName, LOAD_NoWarn | LOAD_Quiet );
-			SetSuppressLogs(false);	
+			bool bTestDone = false;
+			switch(CurrentState)
+			{
+			case EStaticMeshUVTestState::Ready:
+				LoadNextPackage();
+				break;
+			case EStaticMeshUVTestState::PackageLoaded:
+				bTestDone = ProcessCurrentPackage();
+				break;
+			case EStaticMeshUVTestState::WaitingForPackage:
+			default:
+				break;
+			}
+
+			if( bTestDone )
+			{
+				CollectGarbage(RF_Native);
+			}
+			return bTestDone;
 		}
-		if (Package != NULL)
+
+	private:
+
+		/** 
+		 * Callback when a package has finished loading
+		 */
+		void PackageLoadCallback(const FString& PackageName, UPackage* LoadedPackage)
+		{
+			CurrentPackage = LoadedPackage;
+			CurrentState = EStaticMeshUVTestState::PackageLoaded;
+			ContentPackages.RemoveAt(0);
+		}
+
+		/** 
+		 * Loads the next package into memory
+		 */
+		void LoadNextPackage()
 		{
+			CurrentTest->SetSuppressLogs(true);
+			const FString PackageName = ContentPackages[0];
+			CurrentPackage =  FindPackage(NULL, *PackageName);
+			if( CurrentPackage )
+			{
+				//Already loaded, check for static meshes
+				CurrentState = EStaticMeshUVTestState::PackageLoaded;
+				ContentPackages.RemoveAt(0);
+			}
+			else
+			{
+				LoadPackageAsync(PackageName,FLoadPackageAsyncDelegate::CreateRaw(this, &FUVTestHelper::PackageLoadCallback));
+				CurrentState = EStaticMeshUVTestState::WaitingForPackage;
+			}
+		}
+
+		/** 
+		 * Performs the UV test on any static meshes in the current package
+		 */
+		bool ProcessCurrentPackage()
+		{
+			LoadedPackageCount++;
+			CurrentTest->SetSuppressLogs(false);
+
 			for (FObjectIterator ObjIt; ObjIt; ++ObjIt)
 			{
 				UObject* Object = *ObjIt;
 				if (Object != NULL)
 				{
 					//Call CheckLightMapUVs on all static meshes in this package
-					if( Object->IsA(UStaticMesh::StaticClass()) && Object->IsIn(Package) )
+					if( Object->IsA(UStaticMesh::StaticClass()) && Object->IsIn(CurrentPackage) )
 					{
 						UStaticMesh::CheckLightMapUVs((UStaticMesh*)Object,MissingUVMessages,BadUVMessages,ValidUVMessages,true);
 					}
 				}
 			}
+
+			if( ( (LoadedPackageCount % 10) == 0 ) )
+			{
+				CollectGarbage(RF_Native);
+			}
+
+			CurrentState = EStaticMeshUVTestState::Ready;
+
+			//We are done when we are out of packages to process
+			return ContentPackages.Num() == 0;
 		}
-		else
-		{
-			UE_LOG(LogEditorAutomationTests, Error, TEXT("Error loading package %s."), *PackageName );
-		}
-		
-		//Collect garbage on every Nth package to keep memory usage down.
-		if( ( (++PackageID % 10) == 0 ) )
-		{
-			CollectGarbage(RF_Native);
-		}
-	}
-	CollectGarbage(RF_Native);
+
+		// List of packages to test
+		TArray<FString> ContentPackages;
+		// The state of the current test (Ready, WaitingForPackage, PackageLoaded) 
+		EStaticMeshUVTestState::Type CurrentState;
+		// The current package we are checking for static meshes
+		UPackage* CurrentPackage;
+		// How many packages we have loaded
+		uint32 LoadedPackageCount;
+		// Pointer to the current test
+		FAutomationTestBase* CurrentTest;
+
+		// Missing UV messages
+		TArray<FString> MissingUVMessages;
+		// Bas UV messages
+		TArray<FString> BadUVMessages;
+		// Valid UV messages
+		TArray<FString> ValidUVMessages;
+	};
+}
+
+/**
+ * Latent command wrapper for the static mesh UV test
+ */
+DEFINE_LATENT_AUTOMATION_COMMAND_ONE_PARAMETER(FPerformUVTestCommand, TSharedPtr<StaticMeshUVTest::FUVTestHelper>, TestInfo);
+
+bool FPerformUVTestCommand::Update()
+{
+	return TestInfo->Update();
+}
+
+/**
+ * StaticMeshUVsTest
+ */
+IMPLEMENT_SIMPLE_AUTOMATION_TEST( FStaticMeshUVsTest, "Editor.Content.Static Mesh UVs", EAutomationTestFlags::ATF_Editor )
+
+/**
+ * Find all static meshes and check the lightmap UVs
+ */
+bool FStaticMeshUVsTest::RunTest(const FString& Parameters)
+{
+	TSharedPtr<StaticMeshUVTest::FUVTestHelper> TestHelper = MakeShareable(new StaticMeshUVTest::FUVTestHelper(this));
+	ADD_LATENT_AUTOMATION_COMMAND(FPerformUVTestCommand(TestHelper));
 	return true;
 }
diff --git a/Engine/Source/Editor/UnrealEd/Private/UnrealEdSrv.cpp b/Engine/Source/Editor/UnrealEd/Private/UnrealEdSrv.cpp
--- a/Engine/Source/Editor/UnrealEd/Private/UnrealEdSrv.cpp
+++ b/Engine/Source/Editor/UnrealEd/Private/UnrealEdSrv.cpp
@@ -1232,6 +1232,69 @@ bool UUnrealEdEngine::Exec( UWorld* InWorld, const TCHAR* Stream, FOutputDevice&
 
 		return true;
 	}
+	else if( FParse::Command(&Str, TEXT("EditorShot")) || FParse::Command(&Str, TEXT("EditorScreenShot")) )
+	{
+		struct Local
+		{
+			static void TakeScreenShotOfWidget( TSharedRef<SWidget> InWidget )
+			{
+				TArray<FColor> OutImageData;
+				FIntVector OutImageSize;
+				FSlateApplication::Get().TakeScreenshot(InWidget,OutImageData,OutImageSize);
+
+				FString FileName;
+				const FString BaseFileName = FPaths::ScreenShotDir() / TEXT("EditorScreenshot");
+				FFileHelper::GenerateNextBitmapFilename(*BaseFileName, FileName);
+				FFileHelper::CreateBitmap(*FileName, OutImageSize.X, OutImageSize.Y, OutImageData.GetTypedData());
+			}
+		};
+
+		if( FSlateApplication::IsInitialized() )
+		{
+			if( FParse::Command(&Str, TEXT("All") ))
+			{
+				TArray< TSharedRef<SWindow> > OpenWindows;
+				FSlateApplication::Get().GetAllVisibleWindowsOrdered(OpenWindows);
+				for( int32 WindowId = 0; WindowId < OpenWindows.Num(); ++WindowId )
+				{
+					Local::TakeScreenShotOfWidget(OpenWindows[WindowId]);
+				}
+			}
+			else
+			{
+				FString WindowNameStr;
+				if ( FParse::Value(Str, TEXT("Name="), WindowNameStr) )
+				{
+					TArray< TSharedRef<SWindow> > OpenWindows;
+					FSlateApplication::Get().GetAllVisibleWindowsOrdered(OpenWindows);
+					for( int32 WindowId = 0; WindowId < OpenWindows.Num(); ++WindowId )
+					{
+						FString CurrentWindowName = OpenWindows[WindowId]->GetTitle().ToString();
+
+						//Strip off the * from the end if it exists
+						if( CurrentWindowName.EndsWith(TEXT("*")) )
+						{
+							CurrentWindowName = CurrentWindowName.LeftChop(1);
+						}
+
+						if( CurrentWindowName == WindowNameStr )
+						{
+							Local::TakeScreenShotOfWidget(OpenWindows[WindowId]);
+						}
+					}
+				}
+				else
+				{
+					TSharedPtr<SWindow> ActiveWindow = FSlateApplication::Get().GetActiveTopLevelWindow();
+					if( ActiveWindow.IsValid() )
+					{
+						Local::TakeScreenShotOfWidget(ActiveWindow.ToSharedRef());
+					}
+				}
+			}
+		}
+		return true;
+	}
 	return false;
 }
 
diff --git a/Engine/Source/Runtime/AutomationMessages/Classes/AutomationWorkerMessages.h b/Engine/Source/Runtime/AutomationMessages/Classes/AutomationWorkerMessages.h
--- a/Engine/Source/Runtime/AutomationMessages/Classes/AutomationWorkerMessages.h
+++ b/Engine/Source/Runtime/AutomationMessages/Classes/AutomationWorkerMessages.h
@@ -317,6 +317,18 @@ struct FAutomationWorkerRunTests
 	UPROPERTY()
 	FString TestName;
 
+	/**
+	 * If true, we will save out screenshots for tests that support them.
+	 */
+	UPROPERTY()
+	bool bScreenshotsEnabled;
+
+	/**
+	 * If true, we will not resize screen shots.
+	 */
+	UPROPERTY()
+	bool bUseFullSizeScreenShots;
+
 
 	/**
 	 * Default constructor.
@@ -326,10 +338,12 @@ struct FAutomationWorkerRunTests
 	/**
 	 * Creates and initializes a new instance.
 	 */
-	FAutomationWorkerRunTests( uint32 InExecutionCount, int32 InRoleIndex, FString InTestName )
+	FAutomationWorkerRunTests( uint32 InExecutionCount, int32 InRoleIndex, FString InTestName, bool InScreenshotsEnabled, bool InFullSizeScreenshots )
 		: ExecutionCount(InExecutionCount)
 		, RoleIndex(InRoleIndex)
 		, TestName(InTestName)
+		, bScreenshotsEnabled(InScreenshotsEnabled)
+		, bUseFullSizeScreenShots(InFullSizeScreenshots)
 	{ }
 };
 
diff --git a/Engine/Source/Runtime/AutomationWorker/Private/AutomationWorkerModule.cpp b/Engine/Source/Runtime/AutomationWorker/Private/AutomationWorkerModule.cpp
--- a/Engine/Source/Runtime/AutomationWorker/Private/AutomationWorkerModule.cpp
+++ b/Engine/Source/Runtime/AutomationWorker/Private/AutomationWorkerModule.cpp
@@ -124,6 +124,8 @@ void FAutomationWorkerModule::Initialize()
 		{
 			GEngine->GameViewport->OnPNGScreenshotCaptured().BindRaw(this, &FAutomationWorkerModule::HandleScreenShotCaptured);
 		}
+		//Register the editor screen shot callback
+		FAutomationTestFramework::GetInstance().OnScreenshotCaptured().BindRaw(this, &FAutomationWorkerModule::HandleScreenShotCaptured);
 #endif
 
 		bExecuteNextNetworkCommand = true;		
@@ -288,23 +290,56 @@ void FAutomationWorkerModule::HandleRequestTestsMessage( const FAutomationWorker
 #if WITH_ENGINE
 void FAutomationWorkerModule::HandleScreenShotCaptured( int32 Width, int32 Height, const TArray<FColor>& Bitmap, const FString& ScreenShotName )
 {
-	const int32 ThumbnailWidth = 256;
-	const int32 ThumbnailHeight = 128;
-	TArray<FColor> ScaledBitmap;
+	if( FAutomationTestFramework::GetInstance().AreScreenshotsEnabled() )
+	{
+		int32 NewHeight = Height;
+		int32 NewWidth = Width;
+
+		TArray<FColor> ScaledBitmap;
+
+		if( FAutomationTestFramework::GetInstance().ShouldUseFullSizeScreenshots() )
+		{
+			ScaledBitmap = Bitmap;
+			//Clear the alpha channel before saving
+			for ( int32 Index = 0; Index < Width*Height; Index++ )
+			{
+				ScaledBitmap[Index].A = 255;
+			}
+		}
+		else
+		{
+			//Set the thumbnail size
+			NewHeight = 128;
+			NewWidth = 256;
 
-	// Create and save the thumbnail
-	FImageUtils::CropAndScaleImage(Width, Height, ThumbnailWidth, ThumbnailHeight, Bitmap, ScaledBitmap);
+			// Create and save the thumbnail
+			FImageUtils::CropAndScaleImage(Width, Height, NewWidth, NewHeight, Bitmap, ScaledBitmap);
+		}
 
-	TArray<uint8> CompressedBitmap;
-	FImageUtils::CompressImageArray(ThumbnailWidth, ThumbnailHeight, ScaledBitmap, CompressedBitmap);
 
-	// Send the screen shot
-	FAutomationWorkerScreenImage* Message = new FAutomationWorkerScreenImage();
+		TArray<uint8> CompressedBitmap;
+		FImageUtils::CompressImageArray(NewWidth, NewHeight, ScaledBitmap, CompressedBitmap);
+
+		// Send the screen shot if we have a target
+		if( TestRequesterGUID.IsValid() )
+		{
+			FAutomationWorkerScreenImage* Message = new FAutomationWorkerScreenImage();
+
+			FString SFilename = ScreenShotName;
+			Message->ScreenShotName = SFilename;
+			Message->ScreenImage = CompressedBitmap;
+			MessageEndpoint->Send(Message, TestRequesterGUID);
+		}
+		else
+		{
+			//Save locally
+			const bool bTree = true;
+			const FString FileName = FPaths::RootDir() + ScreenShotName;
+			IFileManager::Get().MakeDirectory( *FPaths::GetPath(FileName), bTree );
+			FFileHelper::SaveArrayToFile( CompressedBitmap, *FileName );
+		}
+	}
 
-	FString SFilename = ScreenShotName;
-	Message->ScreenShotName = SFilename;
-	Message->ScreenImage = CompressedBitmap;
-	MessageEndpoint->Send(Message, TestRequesterGUID);
 }
 #endif
 
@@ -313,6 +348,7 @@ void FAutomationWorkerModule::HandleRunTestsMessage( const FAutomationWorkerRunT
 	ExecutionCount = Message.ExecutionCount;
 	TestName = Message.TestName;
 	TestRequesterGUID = Context->GetSender();
+	FAutomationTestFramework::GetInstance().SetScreenshotOptions(Message.bScreenshotsEnabled, Message.bUseFullSizeScreenShots);
 
 	// Always allow the first network command to execute
 	bExecuteNextNetworkCommand = true;
diff --git a/Engine/Source/Runtime/Core/Private/Misc/AutomationTest.cpp b/Engine/Source/Runtime/Core/Private/Misc/AutomationTest.cpp
--- a/Engine/Source/Runtime/Core/Private/Misc/AutomationTest.cpp
+++ b/Engine/Source/Runtime/Core/Private/Misc/AutomationTest.cpp
@@ -436,6 +436,27 @@ void FAutomationTestFramework::SetVisualCommandletFilter(const bool bInVisualCom
 	bVisualCommandletFilterOn = bInVisualCommandletFilterOn;
 }
 
+FOnTestScreenshotCaptured& FAutomationTestFramework::OnScreenshotCaptured()
+{
+	return TestScreenshotCapturedDelegate;
+}
+
+void FAutomationTestFramework::SetScreenshotOptions( const bool bInScreenshotsEnabled, const bool bInUseFullSizeScreenshots )
+{
+	bScreenshotsEnabled = bInScreenshotsEnabled;
+	bUseFullSizeScreenShots = bInUseFullSizeScreenshots;
+}
+
+bool FAutomationTestFramework::AreScreenshotsEnabled() const
+{
+	return bScreenshotsEnabled;
+}
+
+bool FAutomationTestFramework::ShouldUseFullSizeScreenshots() const
+{
+	return bUseFullSizeScreenShots;
+}
+
 void FAutomationTestFramework::PrepForAutomationTests()
 {
 	check(!GIsAutomationTesting);
@@ -597,6 +618,7 @@ FAutomationTestFramework::FAutomationTestFramework()
 ,	CurrentTest(NULL)
 ,	bDeveloperDirectoryIncluded(false)
 ,	bVisualCommandletFilterOn(false)
+,	bUseFullSizeScreenShots(false)
 ,	NetworkRoleIndex(0)
 {
 }
diff --git a/Engine/Source/Runtime/Core/Public/Misc/AutomationTest.h b/Engine/Source/Runtime/Core/Public/Misc/AutomationTest.h
--- a/Engine/Source/Runtime/Core/Public/Misc/AutomationTest.h
+++ b/Engine/Source/Runtime/Core/Public/Misc/AutomationTest.h
@@ -358,6 +358,15 @@ public:
 	virtual void Run() = 0;
 };
 
+/**
+ * Delegate type for when a test screenshot has been captured
+ *
+ * The first parameter is the width.
+ * The second parameter is the height.
+ * The third parameter is the array of bitmap data.
+ * The fourth parameter is the screen shot filename.
+ */
+DECLARE_DELEGATE_FourParams(FOnTestScreenshotCaptured, int32, int32, const TArray<FColor>&, const FString&);
 
 /** Class representing the main framework for running automation tests */
 class CORE_API FAutomationTestFramework
@@ -489,6 +498,28 @@ public:
 	 */
 	void SetVisualCommandletFilter(const bool bInVisualCommandletFilterOn);
 
+	/**
+	 * Accessor for delegate called when a png screenshot is captured 
+	 */
+	FOnTestScreenshotCaptured& OnScreenshotCaptured();
+
+	/**
+	 * Sets screenshot options
+	 * @param bInScreenshotsEnabled - If screenshots are enabled
+	 * @param bInUseFullSizeScreenshots - If true, we won't resize the screenshots
+	 */
+	void SetScreenshotOptions( const bool bInScreenshotsEnabled, const bool bInUseFullSizeScreenshots );
+
+	/**
+	 * Gets if screenshots are enabled
+	 */
+	bool AreScreenshotsEnabled() const;
+
+	/**
+	 * Gets if we are using fulll size screenshots
+	 */
+	bool ShouldUseFullSizeScreenshots() const;
+
 private:
 
 	/** Special feedback context used exclusively while automation testing */
@@ -629,8 +660,17 @@ private:
 	/** Whether we want to convert Anim Blueprint **/
 	bool bVisualCommandletFilterOn;
 
+	/** Wheather screenshots are enabled */
+	bool bScreenshotsEnabled;
+
+	/** Wheather we should resize screenshots or not */
+	bool bUseFullSizeScreenShots;
+
 	/** Participation role as given by the automation controller */
 	uint32 NetworkRoleIndex;
+
+	/** Delegate called at the end of the frame when a screenshot is captured and a .png is requested */
+	FOnTestScreenshotCaptured TestScreenshotCapturedDelegate;
 };
 
 
diff --git a/Engine/Source/Runtime/Engine/Classes/Engine/EngineTypes.h b/Engine/Source/Runtime/Engine/Classes/Engine/EngineTypes.h
--- a/Engine/Source/Runtime/Engine/Classes/Engine/EngineTypes.h
+++ b/Engine/Source/Runtime/Engine/Classes/Engine/EngineTypes.h
@@ -2424,6 +2424,21 @@ struct FFilePath
 	FString FilePath;
 };
 
+/**
+ * Structure for directory paths that are displayed in the UI.
+ */
+USTRUCT()
+struct FDirectoryPath
+{
+	GENERATED_USTRUCT_BODY()
+
+	/**
+	 * The path to the directory.
+	 */
+	UPROPERTY(EditAnywhere, Category=Path)
+	FString Path;
+};
+
 /**
 * This is used for redirecting old name to new name
 * We use manually parsing array, but that makes harder to modify from property setting
diff --git a/Engine/Source/Runtime/Engine/Classes/Tests/AutomationTestSettings.h b/Engine/Source/Runtime/Engine/Classes/Tests/AutomationTestSettings.h
--- a/Engine/Source/Runtime/Engine/Classes/Tests/AutomationTestSettings.h
+++ b/Engine/Source/Runtime/Engine/Classes/Tests/AutomationTestSettings.h
@@ -49,16 +49,16 @@ struct FExternalToolDefinition
 	FString CommandLineOptions;
 
 	/* The working directory for the new process. */
-	UPROPERTY(config, EditAnywhere, Category=ExternalTools)
-	FString WorkingDirectory;
+	UPROPERTY(config, EditAnywhere, Category=ExternalTools, meta=(RelativePath))
+	FDirectoryPath WorkingDirectory;
 
 	/* If set, look for scripts with this extension. */
 	UPROPERTY(config, EditAnywhere, Category=ExternalTools)
 	FString ScriptExtension;
 
 	/* If the ScriptExtension is set, look here for the script files. */
-	UPROPERTY(config, EditAnywhere, Category=ExternalTools)
-	FString ScriptDirectory;
+	UPROPERTY(config, EditAnywhere, Category=ExternalTools, meta=(RelativePath))
+	FDirectoryPath ScriptDirectory;
 };
 
 /**
diff --git a/Engine/Source/Runtime/Engine/Private/Tests/EngineAutomationTests.cpp b/Engine/Source/Runtime/Engine/Private/Tests/EngineAutomationTests.cpp
--- a/Engine/Source/Runtime/Engine/Private/Tests/EngineAutomationTests.cpp
+++ b/Engine/Source/Runtime/Engine/Private/Tests/EngineAutomationTests.cpp
@@ -6,6 +6,7 @@
 #include "FileHelpers.h"
 #endif
 
+#include "AutomationCommon.h"
 #include "AutomationTestCommon.h"
 #include "PlatformFeatures.h"
 #include "SaveGameSystem.h"
@@ -117,6 +118,18 @@ bool FPerformanceCaptureTest::RunTest(const FString& Parameters)
 	return true;
 }
 
+/**
+ * Latent command to take a screenshot of the viewport
+ */
+DEFINE_LATENT_AUTOMATION_COMMAND_ONE_PARAMETER(FTakeViewportScreenshotCommand, FString, ScreenshotFileName);
+
+bool FTakeViewportScreenshotCommand::Update()
+{
+	
+	FScreenshotRequest::RequestScreenshot( ScreenshotFileName, false );
+	return true;
+}
+
 /**
  * LoadAllMapsInGame
  * Verification automation test to make sure loading all maps succeed without crashing AND does performance captures
@@ -163,7 +176,25 @@ bool FLoadAllMapsInGameTest::RunTest(const FString& Parameters)
 {
 	FString MapName = Parameters;
 
+	//Open the map
 	GEngine->Exec(GetSimpleEngineAutomationTestGameWorld(GetTestFlags()), *FString::Printf(TEXT("Open %s"), *MapName));
+
+	if( FAutomationTestFramework::GetInstance().AreScreenshotsEnabled() )
+	{
+		//Generate the screen shot name and path
+		FString ScreenshotFileName;
+		const FString TestName = FString::Printf(TEXT("LoadAllMaps_Game/%s"), *FPaths::GetBaseFilename(MapName));
+		AutomationCommon::GetScreenshotPath(TestName, ScreenshotFileName, true);
+
+		//Give the map some time to load
+		ADD_LATENT_AUTOMATION_COMMAND(FEngineWaitLatentCommand(1.5f));
+		//Take the screen shot
+		ADD_LATENT_AUTOMATION_COMMAND(FTakeViewportScreenshotCommand(ScreenshotFileName));
+		//Give the screen shot a chance to capture the scene
+		ADD_LATENT_AUTOMATION_COMMAND(FEngineWaitLatentCommand(0.5f));
+	}
+
+	//Kick off any Automation matinees that are in this map
 	ADD_LATENT_AUTOMATION_COMMAND(FEnqueuePerformanceCaptureCommands());
 
 	return true;
@@ -222,6 +253,20 @@ bool FSaveGameTest::RunTest(const FString& Parameters)
 	return LoadedData == SavedData;
 }
 
+/**
+ * Latent command to load a map in game
+ */
+DEFINE_LATENT_AUTOMATION_COMMAND_ONE_PARAMETER(FLoadGameMapCommand, FString, MapName);
+
+bool FLoadGameMapCommand::Update()
+{
+	check(GEngine->GetWorldContexts().Num() == 1);
+	check(GEngine->GetWorldContexts()[0].WorldType == EWorldType::Game);
+
+	GEngine->Exec(GEngine->GetWorldContexts()[0].World(), *FString::Printf(TEXT("Open %s"), *MapName));
+	return true;
+}
+
 
 /* UAutomationTestSettings interface
  *****************************************************************************/
diff --git a/Engine/Source/Runtime/Engine/Private/Tests/ExternalToolTest.cpp b/Engine/Source/Runtime/Engine/Private/Tests/ExternalToolTest.cpp
--- a/Engine/Source/Runtime/Engine/Private/Tests/ExternalToolTest.cpp
+++ b/Engine/Source/Runtime/Engine/Private/Tests/ExternalToolTest.cpp
@@ -93,25 +93,28 @@ void FRunExternalToolTest::GetTests(TArray<FString>& OutBeautifiedNames, TArray
 		if( ToolIt.ExecutablePath.FilePath.Len() > 0 )
 		{
 			// Check for multiple scripts if we have a script extension to look for
-			if( ToolIt.ScriptExtension.Len() > 0 && ToolIt.ScriptDirectory.Len() > 0 )
+			if( ToolIt.ScriptExtension.Len() > 0 && ToolIt.ScriptDirectory.Path.Len() > 0 )
 			{
 				TArray<FString> ScriptFileNames;
 
-				IFileManager::Get().FindFiles(ScriptFileNames, *(ToolIt.ScriptDirectory / TEXT("*.") + ToolIt.ScriptExtension), true, false);
+				//Remove the "." before the extension if it exists
+				const FString ScriptExtension = ToolIt.ScriptExtension.StartsWith(TEXT(".")) ? ToolIt.ScriptExtension.Right(ToolIt.ScriptExtension.Len() - 1) : ToolIt.ScriptExtension;
+
+				IFileManager::Get().FindFiles(ScriptFileNames, *(ToolIt.ScriptDirectory.Path / TEXT("*.") + ScriptExtension), true, false);
 
 				for (TArray<FString>::TConstIterator ScriptIt(ScriptFileNames); ScriptIt; ++ScriptIt)
 				{
 					const FString ScriptName = FString::Printf(TEXT("%s: %s"),*ToolIt.ToolName,*FPaths::GetBaseFilename(*ScriptIt));
 					OutBeautifiedNames.Add(*ScriptName);
 
-					const FString SctiptCommand = FString::Printf(TEXT("%s;%s;%s"),*ToolIt.ExecutablePath.FilePath,**ScriptIt,ToolIt.WorkingDirectory.Len() > 0 ? *ToolIt.WorkingDirectory : *ToolIt.ScriptDirectory);
+					const FString SctiptCommand = FString::Printf(TEXT("%s;%s;%s"),*ToolIt.ExecutablePath.FilePath,**ScriptIt,ToolIt.WorkingDirectory.Path.Len() > 0 ? *ToolIt.WorkingDirectory.Path : *ToolIt.ScriptDirectory.Path);
 					OutTestCommands.Add(*SctiptCommand);
 				}
 			}
 			else
 			{
 				OutBeautifiedNames.Add(*ToolIt.ToolName);
-				const FString ToolCommand = FString::Printf(TEXT("%s;%s;%s"),*ToolIt.ExecutablePath.FilePath,*ToolIt.CommandLineOptions,*ToolIt.WorkingDirectory);
+				const FString ToolCommand = FString::Printf(TEXT("%s;%s;%s"),*ToolIt.ExecutablePath.FilePath,*ToolIt.CommandLineOptions,*ToolIt.WorkingDirectory.Path);
 				OutTestCommands.Add(*ToolCommand);
 			}
 		}
diff --git a/Engine/Source/Runtime/Engine/Private/UnrealEngine.cpp b/Engine/Source/Runtime/Engine/Private/UnrealEngine.cpp
--- a/Engine/Source/Runtime/Engine/Private/UnrealEngine.cpp
+++ b/Engine/Source/Runtime/Engine/Private/UnrealEngine.cpp
@@ -5757,6 +5757,10 @@ void UEngine::PerformanceCapture(const FString& CaptureName)
 	PathName += TEXT("/");
 
 	//mapname/CaptureName/platform/version.png
+
+	//Make path relative to the root.
+	PathName = FPaths::AutomationDir() + PathName;
+	FPaths::MakePathRelativeTo(PathName,*FPaths::RootDir());
 	
 	FString ScreenshotName = FString::Printf(TEXT("%s%d.png"), *PathName, GEngineVersion.GetChangelist());
 	
diff --git a/Engine/Source/Runtime/Engine/Public/Tests/AutomationCommon.cpp b/Engine/Source/Runtime/Engine/Public/Tests/AutomationCommon.cpp
--- a/Engine/Source/Runtime/Engine/Public/Tests/AutomationCommon.cpp
+++ b/Engine/Source/Runtime/Engine/Public/Tests/AutomationCommon.cpp
@@ -2,30 +2,23 @@
 
 #include "EnginePrivate.h"
 #include "AutomationCommon.h"
+#include "ImageUtils.h"
 
 ///////////////////////////////////////////////////////////////////////
-// Common Latent commands which are used across test type. I.e. Engine, Network, etc...
+// Common Latent commands
 
-static void SaveWindowAsScreenshot(TSharedRef<SWindow> Window, const FString& FileName)
+namespace AutomationCommon
 {
-	TSharedRef<SWidget> WindowRef = Window;
-
-	TArray<FColor> OutImageData;
-	FIntVector OutImageSize;
-	FSlateApplication::Get().TakeScreenshot(WindowRef,OutImageData,OutImageSize);
+	/** These save a PNG and get sent over the network */
+	static void SaveWindowAsScreenshot(TSharedRef<SWindow> Window, const FString& FileName)
+	{
+		TSharedRef<SWidget> WindowRef = Window;
 
-	FFileHelper::CreateBitmap(*FileName, OutImageSize.X, OutImageSize.Y, OutImageData.GetTypedData());
-}
+		TArray<FColor> OutImageData;
+		FIntVector OutImageSize;
+		FSlateApplication::Get().TakeScreenshot(WindowRef,OutImageData,OutImageSize);
 
-static void SaveEditorWindowsAsScreenshots(const FString& BaseFileName)
-{
-	TArray<TSharedRef<SWindow> > ChildWindows;
-	FSlateApplication::Get().GetAllVisibleWindowsOrdered(ChildWindows);
-
-	for( int32 i = 0; i < ChildWindows.Num(); ++i )
-	{
-		const FString BMPFileName = FString::Printf(TEXT("%s%02i.bmp"), *BaseFileName, i+1);
-		SaveWindowAsScreenshot(ChildWindows[i],BMPFileName);
+		FAutomationTestFramework::GetInstance().OnScreenshotCaptured().ExecuteIfBound(OutImageSize.X, OutImageSize.Y, OutImageData, FileName);
 	}
 }
 
@@ -41,12 +34,12 @@ bool FWaitLatentCommand::Update()
 
 bool FTakeActiveEditorScreenshotCommand::Update()
 {
-	SaveWindowAsScreenshot(FSlateApplication::Get().GetActiveTopLevelWindow().ToSharedRef(),ScreenshotName);
+	AutomationCommon::SaveWindowAsScreenshot(FSlateApplication::Get().GetActiveTopLevelWindow().ToSharedRef(),ScreenshotName);
 	return true;
 }
 
 bool FTakeEditorScreenshotCommand::Update()
 {
-	SaveWindowAsScreenshot(ScreenshotParameters.CurrentWindow.ToSharedRef(), ScreenshotParameters.ScreenshotName);
+	AutomationCommon::SaveWindowAsScreenshot(ScreenshotParameters.CurrentWindow.ToSharedRef(), ScreenshotParameters.ScreenshotName);
 	return true;
-}
+}
\ No newline at end of file
diff --git a/Engine/Source/Runtime/Engine/Public/Tests/AutomationCommon.h b/Engine/Source/Runtime/Engine/Public/Tests/AutomationCommon.h
--- a/Engine/Source/Runtime/Engine/Public/Tests/AutomationCommon.h
+++ b/Engine/Source/Runtime/Engine/Public/Tests/AutomationCommon.h
@@ -9,6 +9,68 @@
 
 DEFINE_LOG_CATEGORY_STATIC(LogEditorAutomationTests, Log, All);
 
+/** Common automation functions */
+namespace AutomationCommon
+{
+	/** Get a string conatins the render mode we are currently in */
+	static FString GetRenderDetailsString()
+	{
+		FString HardwareDetailsString;
+
+		// Create the folder name based on the hardware specs we have been provided
+		FString HardwareDetails = FHardwareInfo::GetHardwareDetailsString();
+
+		FString RHIString;
+		FString RHILookup = NAME_RHI.ToString() + TEXT( "=" );
+		if( FParse::Value( *HardwareDetails, *RHILookup, RHIString ) )
+		{
+			HardwareDetailsString = ( HardwareDetailsString + TEXT( "_" ) ) + RHIString;
+		}
+
+		FString TextureFormatString;
+		FString TextureFormatLookup = NAME_TextureFormat.ToString() + TEXT( "=" );
+		if( FParse::Value( *HardwareDetails, *TextureFormatLookup, TextureFormatString ) )
+		{
+			HardwareDetailsString = ( HardwareDetailsString + TEXT( "_" ) ) + TextureFormatString;
+		}
+
+		FString DeviceTypeString;
+		FString DeviceTypeLookup = NAME_DeviceType.ToString() + TEXT( "=" );
+		if( FParse::Value( *HardwareDetails, *DeviceTypeLookup, DeviceTypeString ) )
+		{
+			HardwareDetailsString = ( HardwareDetailsString + TEXT( "_" ) ) + DeviceTypeString;
+		}
+
+		FString FeatureLevelString;
+		GetFeatureLevelName(GRHIFeatureLevel,FeatureLevelString);
+		{
+			HardwareDetailsString = ( HardwareDetailsString + TEXT( "_" ) ) + FeatureLevelString;
+		}
+
+		if(GEngine->StereoRenderingDevice.IsValid())
+		{
+			HardwareDetailsString = ( HardwareDetailsString + TEXT( "_" ) ) + TEXT("STEREO");
+		}
+
+		return HardwareDetailsString;
+	}
+
+	/** Gets a path used for automation testing (PNG sent to the AutomationTest folder) */
+	static void GetScreenshotPath(const FString& TestName, FString& OutScreenshotName, const bool bIncludeHardwareDetails)
+	{
+		FString PathName = FPaths::AutomationDir() + TestName / FPlatformProperties::PlatformName();
+
+		if( bIncludeHardwareDetails )
+		{
+			PathName = PathName + GetRenderDetailsString();
+		}
+
+		FPaths::MakePathRelativeTo(PathName, *FPaths::RootDir());
+
+		OutScreenshotName = FString::Printf(TEXT("%s/%d.png"), *PathName, GEngineVersion.GetChangelist());
+	}
+}
+
 /**
  * Parameters to the Latent Automation command FTakeEditorScreenshotCommand
  */
