Automation: Merging //depot/Partners/Zombie/UE4-iOS/... to //depot/UE4/...

Milestone 4
 -Added support for grouping devices based off MachineName, Platform, OSVersion, Model, GPU, CPU, RAM, or RenderMode.
 -Added a graphical test result section to show relative times between instances.

diff --git a/Engine/Source/Developer/AutomationController/Private/AutomationControllerManager.h b/Engine/Source/Developer/AutomationController/Private/AutomationControllerManager.h
--- a/Engine/Source/Developer/AutomationController/Private/AutomationControllerManager.h
+++ b/Engine/Source/Developer/AutomationController/Private/AutomationControllerManager.h
@@ -51,7 +51,7 @@ public:
 		bRequestFullScreenScreenshots = bNewValue;
 	}
 
-	virtual bool AreScreenshotsEnabled() const OVERRIDE
+	virtual bool IsScreenshotAllowed() const OVERRIDE
 	{
 		return bScreenshotsEnabled;
 	}
@@ -81,6 +81,11 @@ public:
 		return DeviceClusterManager.GetNumDevicesInCluster(ClusterIndex);
 	}
 
+	virtual FString GetClusterGroupName(const int32 ClusterIndex) const OVERRIDE
+	{
+		return DeviceClusterManager.GetClusterGroupName(ClusterIndex);
+	}
+
 	virtual FString GetDeviceTypeName(const int32 ClusterIndex) const OVERRIDE
 	{
 		return DeviceClusterManager.GetClusterDeviceType(ClusterIndex);
@@ -186,6 +191,16 @@ public:
 		return TestsRefreshedDelegate;
 	}
 
+	virtual bool IsDeviceGroupFlagSet( EAutomationDeviceGroupTypes::Type InDeviceGroup ) const;
+
+	virtual void ToggleDeviceGroupFlag( EAutomationDeviceGroupTypes::Type InDeviceGroup );
+
+	void UpdateDeviceGroups();
+
+	virtual void SetTestsCompleteCallback(const FOnAutomationControllerTestsComplete& NewCallback) OVERRIDE
+	{
+		TestsCompleteDelegate = NewCallback;
+	}
 	// End IAutomationController Interface
 
 
@@ -253,7 +268,6 @@ protected:
 	 */
 	void UpdateTests();
 
-
 private:
 
 	// Handles FAutomationWorkerFindWorkersResponse messages.
@@ -279,13 +293,15 @@ private:
 
 
 private:
-
 	/** Session this controller is currently communicating with */
 	FGuid ActiveSessionId;
 
 	/** The automation test state */
 	EAutomationControllerModuleState::Type AutomationTestState;
 
+	/** Which grouping flags are enabled */
+	uint32 DeviceGroupFlags;
+
 	/** Whether to include developer content in the automation tests */
 	bool bDeveloperDirectoryIncluded;
 
@@ -375,4 +391,7 @@ private:
 
 	// Holds a delegate that is invoked when the controller's tests are being refreshed.
 	FOnAutomationControllerTestsRefreshed TestsRefreshedDelegate;
+
+	// Holds a delegate that is invoked when the tests have completed
+	FOnAutomationControllerTestsComplete TestsCompleteDelegate;
 };
diff --git a/Engine/Source/Developer/AutomationController/Private/AutomationControllerManger.cpp b/Engine/Source/Developer/AutomationController/Private/AutomationControllerManger.cpp
--- a/Engine/Source/Developer/AutomationController/Private/AutomationControllerManger.cpp
+++ b/Engine/Source/Developer/AutomationController/Private/AutomationControllerManger.cpp
@@ -160,6 +160,9 @@ void FAutomationControllerManager::ProcessAvailableTasks()
 				if ( ClusterDistributionMask == 0 )
 				{
 					ProcessResults();
+
+					//Notify the graphical layout we are done processing results.
+					TestsCompleteDelegate.ExecuteIfBound();
 				}
 			}
 		}
@@ -288,6 +291,10 @@ void FAutomationControllerManager::Startup()
 
 	NumOfTestsToReceive = 0;
 	NumTestPasses = 1;
+
+	//Default to machine name
+	DeviceGroupFlags = 0;
+	ToggleDeviceGroupFlag(EAutomationDeviceGroupTypes::MachineName);
 }
 
 
@@ -304,6 +311,7 @@ void FAutomationControllerManager::RemoveCallbacks()
 	ShutdownDelegate.Unbind();
 	TestsAvailableDelegate.Unbind();
 	TestsRefreshedDelegate.Unbind();
+	TestsCompleteDelegate.Unbind();
 }
 
 
@@ -531,7 +539,6 @@ bool FAutomationControllerManager::IsTestRunnable( IAutomationReportPtr InReport
 	return bIsRunnable;
 }
 
-
 /* FAutomationControllerModule callbacks
  *****************************************************************************/
 
@@ -539,7 +546,7 @@ void FAutomationControllerManager::HandleFindWorkersResponseMessage( const FAuto
 {
 	if (Message.SessionId == ActiveSessionId)
 	{
-		DeviceClusterManager.Add(Context->GetSender(), Message.Platform, Message.InstanceName);
+		DeviceClusterManager.AddDeviceFromMessage(Context->GetSender(), Message, DeviceGroupFlags);
 	}
 
 	RequestTests();
@@ -654,3 +661,24 @@ void FAutomationControllerManager::HandleWorkerOfflineMessage( const FAutomation
 	FMessageAddress DeviceMessageAddress = Context->GetSender();
 	DeviceClusterManager.Remove(DeviceMessageAddress);
 }
+
+bool FAutomationControllerManager::IsDeviceGroupFlagSet( EAutomationDeviceGroupTypes::Type InDeviceGroup ) const
+{
+	const uint32 FlagMask = 1 << InDeviceGroup;
+	return (DeviceGroupFlags & FlagMask) > 0;
+}
+
+void FAutomationControllerManager::ToggleDeviceGroupFlag( EAutomationDeviceGroupTypes::Type InDeviceGroup )
+{
+	const uint32 FlagMask = 1 << InDeviceGroup;
+	DeviceGroupFlags = DeviceGroupFlags ^ FlagMask;
+}
+
+void FAutomationControllerManager::UpdateDeviceGroups( )
+{
+	DeviceClusterManager.ReGroupDevices( DeviceGroupFlags );
+
+	// Update the reports in case the number of clusters changed
+	int32 NumOfClusters = DeviceClusterManager.GetNumClusters();
+	ReportManager.ClustersUpdated(NumOfClusters);
+}
\ No newline at end of file
diff --git a/Engine/Source/Developer/AutomationController/Private/AutomationDeviceClusterManager.cpp b/Engine/Source/Developer/AutomationController/Private/AutomationDeviceClusterManager.cpp
--- a/Engine/Source/Developer/AutomationController/Private/AutomationDeviceClusterManager.cpp
+++ b/Engine/Source/Developer/AutomationController/Private/AutomationDeviceClusterManager.cpp
@@ -12,21 +12,23 @@ void FAutomationDeviceClusterManager::Reset()
 }
 
 
-void FAutomationDeviceClusterManager::Add(const FMessageAddress& MessageAddress, const FString& DeviceTypeName, const FString& GameInstanceName)
+void FAutomationDeviceClusterManager::AddDeviceFromMessage(const FMessageAddress& MessageAddress, const FAutomationWorkerFindWorkersResponse& Message, const uint32 GroupFlags)
 {
 	int32 TestClusterIndex;
 	int32 TestDeviceIndex;
 	//if we don't already know about this device
 	if (!FindDevice(MessageAddress, TestClusterIndex, TestDeviceIndex))
 	{
+		FDeviceState NewDevice(MessageAddress, Message);
+		FString GroupName = GetGroupNameForDevice(NewDevice, GroupFlags);
 		//ensure the proper cluster exists
 		int32 ClusterIndex = 0;
 		for (; ClusterIndex < Clusters.Num(); ++ClusterIndex)
 		{
-			if (Clusters[ClusterIndex].DeviceTypeName == DeviceTypeName)
+			if (Clusters[ClusterIndex].ClusterName == GroupName)
 			{
 				//found the cluster, now just append the device
-				Clusters[ClusterIndex].Devices.Add(FDeviceState(MessageAddress, GameInstanceName));
+				Clusters[ClusterIndex].Devices.Add(NewDevice);
 				break;
 			}
 		}
@@ -34,8 +36,9 @@ void FAutomationDeviceClusterManager::Add(const FMessageAddress& MessageAddress,
 		if (ClusterIndex == Clusters.Num())
 		{
 			FDeviceCluster NewCluster;
-			NewCluster.DeviceTypeName = DeviceTypeName;
-			NewCluster.Devices.Add(FDeviceState(MessageAddress, GameInstanceName));
+			NewCluster.ClusterName = GroupName;
+			NewCluster.DeviceTypeName = Message.Platform;
+			NewCluster.Devices.Add(NewDevice);
 			Clusters.Add(NewCluster);
 		}
 	}
@@ -56,6 +59,98 @@ void FAutomationDeviceClusterManager::Remove(const FMessageAddress& MessageAddre
 	}
 }
 
+FString FAutomationDeviceClusterManager::GetGroupNameForDevice(const FDeviceState& DeviceState, const uint32 DeviceGroupFlags)
+{
+	FString OutGroupName;
+
+	if( (DeviceGroupFlags & (1 << EAutomationDeviceGroupTypes::MachineName)) > 0 )
+	{
+		OutGroupName += DeviceState.DeviceName + TEXT("-");
+	}
+
+	if( (DeviceGroupFlags & (1 << EAutomationDeviceGroupTypes::Platform)) > 0 )
+	{
+		OutGroupName += DeviceState.PlatformName + TEXT("-");
+	}
+
+	if( (DeviceGroupFlags & (1 << EAutomationDeviceGroupTypes::OSVersion)) > 0 )
+	{
+		OutGroupName += DeviceState.OSVersionName + TEXT("-");
+	}
+
+	if( (DeviceGroupFlags & (1 << EAutomationDeviceGroupTypes::Model)) > 0 )
+	{
+		OutGroupName += DeviceState.ModelName + TEXT("-");
+	}
+
+	if( (DeviceGroupFlags & (1 << EAutomationDeviceGroupTypes::GPU)) > 0 )
+	{
+		OutGroupName += DeviceState.GPUName + TEXT("-");
+	}
+
+	if( (DeviceGroupFlags & (1 << EAutomationDeviceGroupTypes::CPUModel)) > 0 )
+	{
+		OutGroupName += DeviceState.CPUModelName + TEXT("-");
+	}
+
+	if( (DeviceGroupFlags & (1 << EAutomationDeviceGroupTypes::RamInGB)) > 0 )
+	{
+		OutGroupName += FString::Printf(TEXT("%uGB Ram-"),DeviceState.RAMInGB);
+	}
+
+	if( (DeviceGroupFlags & (1 << EAutomationDeviceGroupTypes::RenderMode)) > 0 )
+	{
+		OutGroupName += DeviceState.RenderModeName + TEXT("-");
+	}
+
+	if( OutGroupName.Len() > 0 )
+	{
+		//Get rid of the trailing '-'
+		OutGroupName = OutGroupName.LeftChop(1);
+	}
+
+	return OutGroupName;
+}
+
+void FAutomationDeviceClusterManager::ReGroupDevices( const uint32 GroupFlags )
+{
+	//Get all the devices
+	TArray< FDeviceState > AllDevices;
+	for(int32 i=0; i<GetNumClusters(); ++i)
+	{
+		AllDevices += Clusters[i].Devices;
+	}
+
+	//Clear out the clusters
+	Reset();
+
+	//Generate new group names based off the active flags and readd the devices
+	for(int32 i=0; i<AllDevices.Num(); ++i)
+	{
+		const FDeviceState* DeviceIt = &AllDevices[i];
+		FString GroupName = GetGroupNameForDevice(*DeviceIt, GroupFlags);
+		//ensure the proper cluster exists
+		int32 ClusterIndex = 0;
+		for (; ClusterIndex < Clusters.Num(); ++ClusterIndex)
+		{
+			if (Clusters[ClusterIndex].ClusterName == GroupName)
+			{
+				//found the cluster, now just append the device
+				Clusters[ClusterIndex].Devices.Add(*DeviceIt);
+				break;
+			}
+		}
+		// if we didn't find the device type yet, add a new cluster and add this device
+		if (ClusterIndex == Clusters.Num())
+		{
+			FDeviceCluster NewCluster;
+			NewCluster.ClusterName = GroupName;
+			NewCluster.DeviceTypeName = DeviceIt->PlatformName;
+			NewCluster.Devices.Add(*DeviceIt);
+			Clusters.Add(NewCluster);
+		}
+	}
+}
 
 int32 FAutomationDeviceClusterManager::GetNumClusters() const
 {
@@ -94,6 +189,12 @@ int32 FAutomationDeviceClusterManager::GetNumActiveDevicesInCluster(const int32
 	return ActiveDevices;
 }
 
+FString FAutomationDeviceClusterManager::GetClusterGroupName(const int32 ClusterIndex) const
+{
+	check((ClusterIndex >= 0) && (ClusterIndex < Clusters.Num()));
+	return Clusters[ClusterIndex].ClusterName;
+}
+
 FString FAutomationDeviceClusterManager::GetClusterDeviceType(const int32 ClusterIndex) const
 {
 	check((ClusterIndex >= 0) && (ClusterIndex < Clusters.Num()));
diff --git a/Engine/Source/Developer/AutomationController/Private/AutomationDeviceClusterManager.h b/Engine/Source/Developer/AutomationController/Private/AutomationDeviceClusterManager.h
--- a/Engine/Source/Developer/AutomationController/Private/AutomationDeviceClusterManager.h
+++ b/Engine/Source/Developer/AutomationController/Private/AutomationDeviceClusterManager.h
@@ -17,11 +17,12 @@ public:
 	void Reset();
 
 	/**
-	 * Add a new device, creating new clusters as needed
+	 * Add a new device from an AutomationWorkerFindWorkersResponse message, creating new clusters as needed
 	 * @param MessageAddress - The network address of an available worker
-	 * @param DeviceTypeName - The name of the device cluster this worker belongs in
+	 * @param Message - The message that contains the device info
+	 * @param GroupFlags - Used find out which group to put the new device into
 	 */
-	void Add(const FMessageAddress& MessageAddress, const FString& DeviceTypeName, const FString& InstanceName);
+	void AddDeviceFromMessage(const FMessageAddress& MessageAddress, const FAutomationWorkerFindWorkersResponse& Message, const uint32 GroupFlags);
 
 	/** Remove a device (went offline, etc) */
 	void Remove(const FMessageAddress& MessageAddress);
@@ -38,9 +39,13 @@ public:
 	/** Returns the number of active devices in of a particular device type */
 	int32 GetNumActiveDevicesInCluster(const int32 ClusterIndex) const;
 
+	/** Returns the name of the cluster group */
+	FString GetClusterGroupName(const int32 ClusterIndex) const;
+
 	/** Returns the name of the devices within this cluster */
 	FString GetClusterDeviceType (const int32 ClusterIndex) const;
 
+	/** Returns the name of a device within a cluster */
 	FString GetClusterDeviceName(const int32 ClusterIndex, const int32 DeviceIndex) const;
 
 	/** 
@@ -118,17 +123,30 @@ public:
 	 */
 	bool HasActiveDevice();
 
+	/**
+	 * Regroups the device clusters based off the current group flags
+	 */
+	void ReGroupDevices( const uint32 GroupFlags );
+
 private:
 
 	/** per actual device, the network address and current test being run */
 	class FDeviceState
 	{
 	public:
-		FDeviceState(FMessageAddress NewMessageAddress, const FString InGameInstanceName )
+		FDeviceState(FMessageAddress NewMessageAddress, const FAutomationWorkerFindWorkersResponse& Message )
 		{
 			DeviceMessageAddress = NewMessageAddress;
+			DeviceName = Message.DeviceName;
+			PlatformName = Message.Platform;
+			OSVersionName = Message.OSVersionName;
+			ModelName = Message.ModelName;
+			GPUName = Message.GPUName;
+			CPUModelName = Message.CPUModelName;
+			RAMInGB = Message.RAMInGB;
+			RenderModeName = Message.RenderModeName;
 			Report.Reset();
-			GameInstanceName = InGameInstanceName;
+			GameInstanceName = Message.InstanceName;
 			IsDeviceAvailable = true;
 		}
 
@@ -138,6 +156,30 @@ private:
 		/** The instance name */
 		FString GameInstanceName;
 
+		/** The name of device */
+		FString DeviceName;
+
+		/** The name of the platform */
+		FString PlatformName;
+
+		/** The name of the operating system version */
+		FString OSVersionName;
+
+		/** The name of the device model */
+		FString ModelName;
+
+		/** The name of the GPU */
+		FString GPUName;
+
+		/** The name of the CPU model */
+		FString CPUModelName;
+
+		/** The amount of RAM this device has in gigabytes */
+		uint32 RAMInGB;
+
+		/** The name of the current render mode */
+		FString RenderModeName;
+
 		/** NULL if this device is available to do work*/
 		TSharedPtr <IAutomationReport> Report;
 
@@ -149,6 +191,9 @@ private:
 	class FDeviceCluster
 	{
 	public:
+		/** Name of the Cluster */
+		FString ClusterName;
+
 		/** Name of the platform */
 		FString DeviceTypeName;
 
@@ -156,6 +201,11 @@ private:
 		TArray <FDeviceState> Devices;
 	};
 
+	/**
+	 * Generates a group name based off the device info and current group flags
+	 */
+	FString GetGroupNameForDevice(const FDeviceState& DeviceState, const uint32 DeviceGroupFlags);
+
 	/** Array of all clusters */
 	TArray <FDeviceCluster> Clusters;
 };
diff --git a/Engine/Source/Developer/AutomationController/Private/AutomationReport.cpp b/Engine/Source/Developer/AutomationController/Private/AutomationReport.cpp
--- a/Engine/Source/Developer/AutomationController/Private/AutomationReport.cpp
+++ b/Engine/Source/Developer/AutomationController/Private/AutomationReport.cpp
@@ -268,6 +268,40 @@ TArray<TSharedPtr<IAutomationReport> >& FAutomationReport::GetChildReports()
 	return ChildReports;
 }
 
+void FAutomationReport::ClustersUpdated(const int32 NumClusters)
+{
+	TestInfo.ResetNumDevicesRunningTest();
+
+	//Fixup Support flags
+	SupportFlags = 0;
+	for (int32 i = 0; i <= NumClusters; ++i)
+	{
+		SupportFlags |= (1<<i);
+	}
+
+	//Fixup Results array
+	if( NumClusters > Results.Num() )
+	{
+		for( int32 ClusterIndex = Results.Num(); ClusterIndex < NumClusters; ++ClusterIndex )
+		{
+			//Make sure we have enough results for a single pass
+			TArray<FAutomationTestResults> AutomationTestResult;
+			AutomationTestResult.Add( FAutomationTestResults() );
+			Results.Add( AutomationTestResult );
+		}
+	}
+	else if( NumClusters > Results.Num() )
+	{
+		Results.RemoveAt(NumClusters, Results.Num() - NumClusters);
+	}
+
+	//recurse to children
+	for (int32 ChildIndex = 0; ChildIndex < ChildReports.Num(); ++ChildIndex)
+	{
+		ChildReports[ChildIndex]->ClustersUpdated(NumClusters);
+	}
+}
+
 void FAutomationReport::ResetForExecution(const int32 NumTestPasses)
 {
 	TestInfo.ResetNumDevicesRunningTest();
diff --git a/Engine/Source/Developer/AutomationController/Private/AutomationReport.h b/Engine/Source/Developer/AutomationController/Private/AutomationReport.h
--- a/Engine/Source/Developer/AutomationController/Private/AutomationReport.h
+++ b/Engine/Source/Developer/AutomationController/Private/AutomationReport.h
@@ -86,6 +86,9 @@ public:
 	/** Returns the array of child reports */
 	virtual TArray<TSharedPtr<IAutomationReport> >& GetChildReports() OVERRIDE;
 
+	/** Updates the report when the number of clusters changes */
+	virtual void ClustersUpdated(const int32 NumClusters) OVERRIDE;
+
 	/** 
 	 * Recursively resets the report to "needs to be run", clears cached warnings and errors
 	 * @param NumTestPasses - The number of test passes so we know how many results to create
diff --git a/Engine/Source/Developer/AutomationController/Private/AutomationReportManager.cpp b/Engine/Source/Developer/AutomationController/Private/AutomationReportManager.cpp
--- a/Engine/Source/Developer/AutomationController/Private/AutomationReportManager.cpp
+++ b/Engine/Source/Developer/AutomationController/Private/AutomationReportManager.cpp
@@ -21,6 +21,10 @@ void FAutomationReportManager::Empty()
 	ReportRoot->Empty();
 }
 
+void FAutomationReportManager::ClustersUpdated( const int32 NumClusters )
+{
+	ReportRoot->ClustersUpdated(NumClusters);
+}
 
 void FAutomationReportManager::ResetForExecution(const int32 NumTestPasses)
 {
diff --git a/Engine/Source/Developer/AutomationController/Private/AutomationReportManager.h b/Engine/Source/Developer/AutomationController/Private/AutomationReportManager.h
--- a/Engine/Source/Developer/AutomationController/Private/AutomationReportManager.h
+++ b/Engine/Source/Developer/AutomationController/Private/AutomationReportManager.h
@@ -19,6 +19,9 @@ public:
 	/** Clears out all reports */
 	void Empty();
 
+	/** Updates the automation reports when the number of clusters changes */
+	void ClustersUpdated( const int32 NumClusters );
+
 	/** 
 	 * Resets for the next run of tests 
 	 * @param NumTestPasses - The number of test results to store
diff --git a/Engine/Source/Developer/AutomationController/Public/Interfaces/IAutomationControllerManager.h b/Engine/Source/Developer/AutomationController/Public/Interfaces/IAutomationControllerManager.h
--- a/Engine/Source/Developer/AutomationController/Public/Interfaces/IAutomationControllerManager.h
+++ b/Engine/Source/Developer/AutomationController/Public/Interfaces/IAutomationControllerManager.h
@@ -86,6 +86,56 @@ namespace EFileExportType
 
 };
 
+namespace EAutomationDeviceGroupTypes
+{
+	enum Type
+	{
+		MachineName,	// Group by machine name
+		Platform,		// Group by platform
+		OSVersion,		// Group by operating system version
+		Model,			// Group by machine model
+		GPU,			// Group by GPU
+		CPUModel,		// Group by CPU Model
+		RamInGB,		// Group by RAM in gigabytes
+		RenderMode,		// Group by RenderMode (D3D11_SM5, OpenGL_SM4, etc)
+		Max
+	};
+
+	static FText ToName( const Type DeviceGroupType )
+	{
+		switch( DeviceGroupType )
+		{
+		case MachineName:	return NSLOCTEXT("AutomationTest", "AutomationDeviceGroup_MachineName", "Machine Name");
+		case Platform:		return NSLOCTEXT("AutomationTest", "AutomationDeviceGroup_Platform", "Platform");
+		case OSVersion:		return NSLOCTEXT("AutomationTest", "AutomationDeviceGroup_OSVersion", "OS Version");
+		case Model:			return NSLOCTEXT("AutomationTest", "AutomationDeviceGroup_Model", "Model");
+		case GPU:			return NSLOCTEXT("AutomationTest", "AutomationDeviceGroup_GPU", "GPU");
+		case CPUModel:		return NSLOCTEXT("AutomationTest", "AutomationDeviceGroup_CPUModel", "CPU Model");
+		case RamInGB:		return NSLOCTEXT("AutomationTest", "AutomationDeviceGroup_RAM", "RAM in GB");
+		case RenderMode:	return NSLOCTEXT("AutomationTest", "AutomationDeviceGroup_RenderMode", "Render Mode");
+
+		default:			return NSLOCTEXT("AutomationTest", "AutomationDeviceGroup_InvalidOrMax", "InvalidOrMax");
+		}
+	}
+
+	static FText ToDescription(const Type DeviceGroupType)
+	{
+		switch( DeviceGroupType )
+		{
+		case MachineName:	return NSLOCTEXT("AutomationTest", "AutomationDeviceGroupTip_MachineName", "Group devices based off their machine name");
+		case Platform:		return NSLOCTEXT("AutomationTest", "AutomationDeviceGroupTip_Platform", "Group devices based off their platform");
+		case OSVersion:		return NSLOCTEXT("AutomationTest", "AutomationDeviceGroupTip_OSVersion", "Group devices based off their OS version");
+		case Model:			return NSLOCTEXT("AutomationTest", "AutomationDeviceGroupTip_Model", "Group devices based off their device model");
+		case GPU:			return NSLOCTEXT("AutomationTest", "AutomationDeviceGroupTip_GPU", "Group devices based off their GPU");
+		case CPUModel:		return NSLOCTEXT("AutomationTest", "AutomationDeviceGroupTip_CPUModel", "Group devices based off their CPU model");
+		case RamInGB:		return NSLOCTEXT("AutomationTest", "AutomationDeviceGroupTip_RAM", "Group devices based off memory (in GB)");
+		case RenderMode:	return NSLOCTEXT("AutomationTest", "AutomationDeviceGroupTip_RenderMode", "Group devices based off the current render mode");
+
+		default:			return NSLOCTEXT("AutomationTest", "AutomationDeviceGroup_InvalidOrMax", "InvalidOrMax");
+		}
+	}
+}
+
 
 /** Console command delegate type (takes no arguments.)  This is a void callback function. */
 DECLARE_DELEGATE(FOnAutomationControllerTestsRefreshed);
@@ -96,6 +146,9 @@ DECLARE_DELEGATE_OneParam(FOnAutomationControllerManagerTestsAvailable, EAutomat
 /** Delegate to call when the automation controller is shutdown. */
 DECLARE_DELEGATE(FOnAutomationControllerManagerShutdown);
 
+/** Delegate to call when the tests are complete. */
+DECLARE_DELEGATE(FOnAutomationControllerTestsComplete);
+
 
 /**
  * Interface for AutomationController .
@@ -163,9 +216,9 @@ public:
 	virtual void SetUsingFullSizeScreenshots(const  bool bNewValue ) = 0;
 
 	/**
-	 * Returns if screenshots are enabled
+	 * Returns if screenshots are allowed
 	 */
-	virtual bool AreScreenshotsEnabled() const = 0;
+	virtual bool IsScreenshotAllowed() const = 0;
 
 	/**
 	 * Sets if screenshots are enabled
@@ -192,6 +245,11 @@ public:
 	 */
 	virtual int32 GetNumDevicesInCluster(const int32 ClusterIndex) const = 0;
 
+	/**
+	 * Get the group name of the specified cluster.
+	 */
+	virtual FString GetClusterGroupName(const int32 ClusterIndex) const =0;
+
 	/**
 	 * Get name of a particular device cluster.
 	 */
@@ -322,7 +380,25 @@ public:
 	 */
 	virtual void Startup() = 0;
 
+	/**
+	 * Checks if a device group flag is set.
+	 */
+	virtual bool IsDeviceGroupFlagSet( EAutomationDeviceGroupTypes::Type InDeviceGroup ) const = 0;
+
+	/**
+	 * Toggles a device group flag.
+	 */
+	virtual void ToggleDeviceGroupFlag( EAutomationDeviceGroupTypes::Type InDeviceGroup ) = 0;
 
+	/**
+	 * Updates the clusters when the device grouping changes.
+	 */
+	virtual void UpdateDeviceGroups( ) = 0;
+
+	/**
+	 * Sets the test complete callback delegate.
+	 */
+	virtual void SetTestsCompleteCallback(const FOnAutomationControllerTestsComplete& NewCallback) = 0;
 public:
 
 	/**
diff --git a/Engine/Source/Developer/AutomationController/Public/Interfaces/IAutomationReport.h b/Engine/Source/Developer/AutomationController/Public/Interfaces/IAutomationReport.h
--- a/Engine/Source/Developer/AutomationController/Public/Interfaces/IAutomationReport.h
+++ b/Engine/Source/Developer/AutomationController/Public/Interfaces/IAutomationReport.h
@@ -203,6 +203,9 @@ public:
 	/** Returns the array of child reports */
 	virtual TArray<TSharedPtr<IAutomationReport> >& GetChildReports() = 0;
 
+	/** Updates the report when the number of clusters changes */
+	virtual void ClustersUpdated(const int32 NumClusters) = 0;
+
 	/** 
 	 * Recursively resets the report to "needs to be run", clears cached warnings and errors
 	 * @param NumTestPasses - The number of test passes so we know how many results to create
diff --git a/Engine/Source/Developer/AutomationWindow/Private/AutomationWindowPrivatePCH.h b/Engine/Source/Developer/AutomationWindow/Private/AutomationWindowPrivatePCH.h
--- a/Engine/Source/Developer/AutomationWindow/Private/AutomationWindowPrivatePCH.h
+++ b/Engine/Source/Developer/AutomationWindow/Private/AutomationWindowPrivatePCH.h
@@ -24,6 +24,8 @@
  *****************************************************************************/
 
 #include "SAutomationWindowCommandBar.h"
+#include "SAutomationGraphicalResultBox.h"
+#include "SAutomationTestTreeView.h"
 #include "SAutomationExportMenu.h"
 #include "AutomationFilter.h"
 #include "AutomationPresetManager.h"
diff --git a/Engine/Source/Developer/AutomationWindow/Private/SAutomationGraphicalResultBox.cpp b/Engine/Source/Developer/AutomationWindow/Private/SAutomationGraphicalResultBox.cpp
new file mode 100644
--- /dev/null
+++ b/Engine/Source/Developer/AutomationWindow/Private/SAutomationGraphicalResultBox.cpp
@@ -0,0 +1,325 @@
+// Copyright 1998-2014 Epic Games, Inc. All Rights Reserved.
+
+/*=============================================================================
+	SAutomationGraphicalResultBox.cpp: Implements the SAutomationGraphicalResultBox class.
+=============================================================================*/
+
+#include "AutomationWindowPrivatePCH.h"
+
+#define LOCTEXT_NAMESPACE "SAutomationGraphicalResultBox"
+
+/* SAutomationGraphicalResultBox interface
+ *****************************************************************************/
+
+void SAutomationGraphicalResultBox::Construct( const FArguments& InArgs, const IAutomationControllerManagerRef& InAutomationController )
+{
+	AutomationController = InAutomationController;
+	DisplayType = EAutomationGrapicalDisplayType::DisplayName;
+
+	RootBox = SNew(SVerticalBox); 
+
+	ChildSlot
+	[
+		SNew(SHorizontalBox)
+		+SHorizontalBox::Slot()
+		.FillWidth(1.0f)
+		[
+			RootBox.ToSharedRef()
+		]
+	];
+
+	ClearResults();
+
+	AutomationController->SetTestsCompleteCallback(FOnAutomationControllerTestsComplete::CreateRaw(this, &SAutomationGraphicalResultBox::OnFillResults));
+}
+
+void SAutomationGraphicalResultBox::ClearResults()
+{
+	TotalTestDuration = 0.f;
+	ClusterResults.Empty();
+	RootBox->ClearChildren();
+}
+
+void SAutomationGraphicalResultBox::GetEnabledReports(TSharedPtr<IAutomationReport> InReport, TArray< TSharedPtr< IAutomationReport > >& outReports)
+{
+	TArray<TSharedPtr<IAutomationReport> >& ChildReports = InReport->GetChildReports();
+
+	if (ChildReports.Num() > 0)
+	{
+		for (int32 Index = 0; Index < ChildReports.Num(); Index++)
+		{
+			GetEnabledReports(ChildReports[Index], outReports);
+		}
+	}
+	else if ( InReport->IsEnabled() )
+	{
+		outReports.Add(InReport);
+	}
+}
+
+void SAutomationGraphicalResultBox::PopulateData()
+{
+	//Get the report list
+	TArray< TSharedPtr< IAutomationReport > > AllReports;
+	AllReports = AutomationController->GetReports();
+
+	//Find only the enabled tests
+	TArray< TSharedPtr< IAutomationReport > > EnabledReports;
+	for( int32 i=0; i<AllReports.Num(); ++i )
+	{
+		GetEnabledReports(AllReports[i],EnabledReports);
+	}
+
+	uint32 LastPassIndex = AutomationController->GetNumPasses() - 1;
+
+	//Pull the data out of the reports
+	for(int i = 0; i < AutomationController->GetNumDeviceClusters(); ++i)
+	{
+		ClusterResults.Add(FClusterResults(AutomationController->GetClusterGroupName(i)));
+		FClusterResults* ClusterIt = &ClusterResults[i];
+		ClusterIt->TotalNumTests = EnabledReports.Num();
+		for(int j = 0; j < AutomationController->GetNumDevicesInCluster(i); ++j)
+		{
+			ClusterIt->Devices.Add(FDeviceResults(AutomationController->GetGameInstanceName(i, j)));
+			FDeviceResults* DeviceIt = &ClusterIt->Devices[j];
+			for( int k = 0; k < EnabledReports.Num(); ++k )
+			{
+				FAutomationTestResults TestResults = EnabledReports[k]->GetResults( i, LastPassIndex );
+				if( TestResults.GameInstance == DeviceIt->InstanceName )
+				{
+					uint32 TestIndex = DeviceIt->Tests.Add(FTestResults(EnabledReports[k]->GetDisplayNameWithDecoration()));
+					FTestResults* TestIt = &DeviceIt->Tests[TestIndex];
+					TestIt->Duration = TestResults.Duration;
+					TestIt->TestState = TestResults.State;
+					TestIt->bHasWarnings = TestResults.Warnings.Num() > 0;
+
+					DeviceIt->TotalTime += TestIt->Duration;
+					ClusterIt->TotalTime += TestIt->Duration;
+
+					if( TestResults.State == EAutomationState::Success )
+					{
+						DeviceIt->TotalTestSuccesses++;
+						ClusterIt->TotalTestSuccesses++;
+					}
+				}
+			}
+
+			//See if this is the new longest running device for this cluster
+			if( DeviceIt->TotalTime > ClusterIt->ParallelTime )
+			{
+				ClusterIt->ParallelTime = DeviceIt->TotalTime;
+			}
+		}
+
+		//The total test duration is the time of the longest cluster.
+		if( ClusterIt->ParallelTime > TotalTestDuration )
+		{
+			TotalTestDuration = ClusterIt->ParallelTime;
+		}
+	}
+}
+
+void SAutomationGraphicalResultBox::OnFillResults()
+{
+	//Clear the old results
+	ClearResults();
+
+	//Get the data out of the automation reports
+	PopulateData();
+
+	//Create widgets for the results
+	CreateWidgets();
+}
+
+void SAutomationGraphicalResultBox::CreateWidgets()
+{
+	//Note: I am using 10 columns here to hack around a sizing bug.  The Grid calculates the width for items that span multiple 
+	//		columns by just dividing the size by the number of columns.  This causes problems if the columns are different sizes.
+	//		By having a bunch of extra columns here it will calculate a smaller per column size for our header.
+	TSharedRef<SGridPanel> GridContainer = SNew(SGridPanel).FillColumn(1,1);
+
+	uint32 RowCounter = 0;
+
+	for( int32 ClusterIndex = 0; ClusterIndex < ClusterResults.Num(); ++ClusterIndex )
+	{
+		FClusterResults* ClusterIt = &ClusterResults[ClusterIndex];
+
+		FFormatNamedArguments ClusterArgs;
+		ClusterArgs.Add(TEXT("Name"), FText::FromString(ClusterIt->ClusterName));
+		ClusterArgs.Add(TEXT("NumTests"), ClusterIt->TotalNumTests);
+		ClusterArgs.Add(TEXT("NumFails"), ClusterIt->TotalNumTests - ClusterIt->TotalTestSuccesses);
+		ClusterArgs.Add(TEXT("TotalTime"), ClusterIt->TotalTime);
+		ClusterArgs.Add(TEXT("ParallelTime"), ClusterIt->ParallelTime);
+
+		//Add Cluster Header
+		GridContainer->AddSlot(0,RowCounter)
+			.ColumnSpan(10)
+			.HAlign(HAlign_Left)
+			.VAlign(VAlign_Center)
+			.Padding( FMargin(1,3) )
+			[
+				SNew(STextBlock)
+				.TextStyle( FEditorStyle::Get(), "Automation.ReportHeader" )
+				.Text(FText::Format(LOCTEXT("AutomationGraphicalClusterHeader", "{Name}  -  {NumTests} Tests / {NumFails} Fails / {TotalTime} Seconds (Total) / {ParallelTime} Seconds (Parallel)"), ClusterArgs))
+			];
+
+		RowCounter++;
+
+		for( int32 DeviceIndex = 0; DeviceIndex < ClusterIt->Devices.Num(); ++DeviceIndex )
+		{
+			FDeviceResults* DeviceIt = &ClusterIt->Devices[DeviceIndex];
+			const int32 NumTests = DeviceIt->Tests.Num();
+
+			//Add Device Header
+			FFormatNamedArguments DeviceArgs;
+			DeviceArgs.Add(TEXT("NumTests"), NumTests);
+			DeviceArgs.Add(TEXT("NumFails"), NumTests - DeviceIt->TotalTestSuccesses);
+			DeviceArgs.Add(TEXT("TotalTime"), DeviceIt->TotalTime);
+
+			GridContainer->AddSlot(0,RowCounter)
+				.HAlign(HAlign_Left)
+				.Padding(0,0,3,0)
+				[
+					SNew(SVerticalBox)
+					+SVerticalBox::Slot()
+					[
+						SNew(STextBlock)
+						.Text(DeviceIt->InstanceName)
+					]
+					+SVerticalBox::Slot()
+					[
+						SNew(STextBlock)
+						.Text(FText::Format(LOCTEXT("AutomationGraphicalDeviceHeader", "{NumTests} Tests / {NumFails} Fails / {TotalTime} Seconds"), DeviceArgs))
+					]
+				];
+
+			TSharedRef<SHorizontalBox> TestContainer = SNew(SHorizontalBox);
+
+			for( int32 TestIndex = 0; TestIndex < DeviceIt->Tests.Num(); ++TestIndex)
+			{
+				FTestResults* TestIt = &DeviceIt->Tests[TestIndex];
+
+				FFormatNamedArguments TestArgs;	
+				TestArgs.Add(TEXT("Duration"), TestIt->Duration);
+				TestArgs.Add(TEXT("Name"), FText::FromString(TestIt->TestName));
+
+				TestContainer->AddSlot()
+					.HAlign(HAlign_Fill)
+					.VAlign(VAlign_Fill)
+					.Padding( FMargin(1,5) )
+					.FillWidth(TestIt->Duration)
+					[
+						SNew(SOverlay)
+						.ToolTipText(FText::Format(LOCTEXT("AutomationGraphicalToolTip", "{Name} \nDuration: {Duration}s"), TestArgs))
+
+						+SOverlay::Slot()
+						[
+							SNew(SBorder)
+							.BorderImage( FEditorStyle::GetBrush("ErrorReporting.Box") )
+							.BorderBackgroundColor(this,&SAutomationGraphicalResultBox::GetColorForTestState, TestIt->TestState, TestIt->bHasWarnings)
+						]
+
+						+ SOverlay::Slot()
+						.HAlign(HAlign_Center)
+						.VAlign(VAlign_Center)
+						[
+							SNew(STextBlock)
+							.Text(this,&SAutomationGraphicalResultBox::GetTestDisplayText,TestIt->TestName,TestIt->Duration)//FText::Format(LOCTEXT("AutomationGraphicalDuration", "{Duration}s"), TestArgs))
+							.ColorAndOpacity(FLinearColor(0,0,0,1))
+						]
+					];
+			}
+
+			//Fill in the end with the remaining time
+			if( DeviceIt->TotalTime < TotalTestDuration )
+			{
+				const float RemainingTime = TotalTestDuration - DeviceIt->TotalTime;
+				TestContainer->AddSlot()
+					.HAlign(HAlign_Fill)
+					.VAlign(VAlign_Fill)
+					.Padding( FMargin(1,5) )
+					.FillWidth( RemainingTime )
+					[
+						SNew(SBorder)
+						.BorderImage( FEditorStyle::GetBrush("ErrorReporting.Box") )
+						.BorderBackgroundColor(FLinearColor(0,0,0,0))
+					];
+			}
+
+			
+			GridContainer->AddSlot(1,RowCounter)
+				.HAlign(HAlign_Fill)
+				.VAlign(VAlign_Center)
+				.Padding( FMargin(1,3) )
+				.ColumnSpan(9)
+				[
+					SNew(SHorizontalBox)
+					+SHorizontalBox::Slot()
+					.FillWidth(1.f)
+					[
+						TestContainer
+					]
+				];
+
+			RowCounter++;
+		}
+	}
+
+	RootBox->AddSlot()
+		[
+			GridContainer
+		];
+}
+
+FSlateColor SAutomationGraphicalResultBox::GetColorForTestState(const EAutomationState::Type TestState, const bool bHasWarnings) const
+{
+	switch(TestState)
+	{
+	case EAutomationState::Success:
+		if (bHasWarnings)
+		{
+			return FSlateColor( FLinearColor( 1.0f, 0.5f, 0.0f ) );
+		}
+		else
+		{
+			return FSlateColor( FLinearColor( 0.0f, 0.5f, 0.0f ) );
+		}
+		break;
+	case EAutomationState::Fail:
+	default: //Default to fail!
+		break;
+	}
+
+	return FSlateColor( FLinearColor( 0.5f, 0.0f, 0.0f ) );
+}
+
+FText SAutomationGraphicalResultBox::GetTestDisplayText(const FString TestName, const float TestTime) const
+{
+	if(DisplayType == EAutomationGrapicalDisplayType::DisplayName)
+	{
+		return FText::FromString(TestName);
+	}
+	else
+	{
+		FFormatNamedArguments TestArgs;	
+		TestArgs.Add(TEXT("Duration"), TestTime);
+		return FText::Format(LOCTEXT("AutomationGraphicalDuration", "{Duration}s"), TestArgs);
+	}
+}
+
+EAutomationGrapicalDisplayType::Type SAutomationGraphicalResultBox::GetDisplayType() const
+{
+	return DisplayType;
+}
+
+void SAutomationGraphicalResultBox::SetDisplayType(EAutomationGrapicalDisplayType::Type NewDisplayType)
+{
+	DisplayType = NewDisplayType;
+}
+
+bool SAutomationGraphicalResultBox::HasResults() const
+{
+	return ClusterResults.Num() > 0;
+}
+
+#undef LOCTEXT_NAMESPACE
diff --git a/Engine/Source/Developer/AutomationWindow/Private/SAutomationGraphicalResultBox.h b/Engine/Source/Developer/AutomationWindow/Private/SAutomationGraphicalResultBox.h
new file mode 100644
--- /dev/null
+++ b/Engine/Source/Developer/AutomationWindow/Private/SAutomationGraphicalResultBox.h
@@ -0,0 +1,203 @@
+// Copyright 1998-2014 Epic Games, Inc. All Rights Reserved.
+
+/*=============================================================================
+	SAutomationGraphicalResultBox.h: Declares the SAutomationGraphicalResultBox class.
+=============================================================================*/
+
+#pragma once
+
+/**
+ * The different supported items we can display on the test item
+ */
+namespace EAutomationGrapicalDisplayType
+{
+	enum Type
+	{
+		DisplayName,
+		DisplayTime
+	};
+}
+
+/**
+ * Implements the automation graphical results box widget.
+ */
+class SAutomationGraphicalResultBox
+	: public SCompoundWidget
+{
+public:
+
+	SLATE_BEGIN_ARGS(SAutomationGraphicalResultBox) { }
+
+	SLATE_ARGUMENT(FString, Text)
+
+	SLATE_END_ARGS()
+
+public:
+
+	/**
+	 * Construct this widget
+	 *
+	 * @param InArgs - The declaration data for this widget.
+	 * @param InAutomationController - Reference to the automation controller
+	 */
+	void Construct( const FArguments& InArgs, const IAutomationControllerManagerRef& InAutomationController);
+
+	/**
+	 * Clears the current results and any widgets that were created
+	 */
+	void ClearResults();
+
+	/**
+	 * Returns if there are any results available to display
+	 */
+	bool HasResults() const;
+
+	/**
+	 * Returns the current display type
+	 */
+	EAutomationGrapicalDisplayType::Type GetDisplayType() const;
+
+	/**
+	 * Sets the current display type
+	 */
+	void SetDisplayType(EAutomationGrapicalDisplayType::Type NewDisplayType);
+private:
+
+	/** Holds information on a single test result */
+	class FTestResults
+	{
+	public:
+		/** Constructor */
+		FTestResults(const FString& InName)
+		{
+			TestName = InName;
+			Duration = 0.f;
+			TestState = EAutomationState::NotRun;
+			bHasWarnings = false;
+		}
+
+		/** Name of the test */
+		FString TestName;
+
+		/** Duration of the test */
+		float Duration;
+
+		/** State of the test */
+		EAutomationState::Type TestState;
+
+		/** If the test had any warnings */
+		bool bHasWarnings;
+	};
+
+	/** Holds all the test results for a single device */
+	class FDeviceResults
+	{
+	public:
+		/** Constructor */
+		FDeviceResults(const FString& InName)
+		{
+			InstanceName = InName;
+			TotalTime = 0.f;
+			TotalTestSuccesses = 0;
+		}
+
+		/** Name of the device instance */
+		FString InstanceName;
+
+		/** Total time of all the tests */
+		float TotalTime;
+
+		/** How many of the tests were successful */
+		uint32 TotalTestSuccesses;
+
+		/** The list of tests */
+		TArray<FTestResults> Tests;
+	};
+
+	/** Holds all the results for a single cluster */
+	class FClusterResults
+	{
+	public:
+		/** Constructor */
+		FClusterResults(const FString& InName)
+		{
+			ClusterName = InName;
+			TotalTime = 0.f;
+			ParallelTime = 0.f;
+			TotalNumTests = 0;
+			TotalTestSuccesses = 0;
+		}
+
+		/** Name of the Cluster */
+		FString ClusterName;
+
+		/** Total number of tests */
+		uint32 TotalNumTests;
+
+		/** How many of the tests were successful */
+		uint32 TotalTestSuccesses;
+
+		/** Total time of all the tests */
+		float TotalTime;
+
+		/** Longest time a device took to finish its tests */
+		float ParallelTime;
+
+		/** The list of devices in this cluster */
+		TArray <FDeviceResults> Devices;
+	};
+
+	/**
+	 * Returns if there are any results available to display
+	 *
+	 * @param InReport - The report to recursively check for enabled tests.
+	 * @param OutReports - The list of enabled test reports
+	 */
+	void GetEnabledReports(TSharedPtr<IAutomationReport> InReport, TArray< TSharedPtr< IAutomationReport > >& OutReports);
+
+	/**
+	 * Populates the ClusterResults with test data pulled from the automation controller
+	 */
+	void PopulateData();
+
+	/**
+	 * Creates all the slate widgets to display the test results
+	 */
+	void CreateWidgets();
+
+	/**
+	 * This is called when the tests are complete so we can update the widget
+	 */
+	void OnFillResults();
+
+	/**
+	 * Gets the color for a test based off the test state and if it has any warnings
+	 *
+	 * @param - TestState - The test state from the automation report (Success, Fail, NotRun, etc.)
+	 * @param - bHasWarnings - If the test had any warning logs.
+	 */
+	FSlateColor GetColorForTestState(const EAutomationState::Type TestState, const bool bHasWarnings) const;
+
+	/**
+	 * Gets the text to display for this test item based off the current DisplayType
+	 *
+	 * @param - TestName - The name of the test
+	 * @param - TestTime - The duration of the test
+	 */
+	FText GetTestDisplayText(const FString TestName, const float TestTime) const;
+
+	/** Stores what information should be displayed on the test widgets */
+	EAutomationGrapicalDisplayType::Type DisplayType;
+
+	/** The test results grouped by cluster */
+	TArray<FClusterResults> ClusterResults;
+
+	/** The time of the longest cluster */
+	float TotalTestDuration;
+
+	/** Pointer to the root widget that contains all the reslut widgets */
+	TSharedPtr< SVerticalBox > RootBox;
+
+	/** Pointer to the automation controller so we can get the test results */
+	IAutomationControllerManagerPtr AutomationController;
+};
diff --git a/Engine/Source/Developer/AutomationWindow/Private/SAutomationTestTreeView.h b/Engine/Source/Developer/AutomationWindow/Private/SAutomationTestTreeView.h
new file mode 100644
--- /dev/null
+++ b/Engine/Source/Developer/AutomationWindow/Private/SAutomationTestTreeView.h
@@ -0,0 +1,147 @@
+// Copyright 1998-2014 Epic Games, Inc. All Rights Reserved.
+
+/*=============================================================================
+	SAutomationTestTreeView.h: Declares the SAutomationTestTreeView class.
+=============================================================================*/
+
+#pragma once
+
+/**
+ * Implements the automation test tree view.  Same slate interface as STreeView but adds the ability to clear the internal widget cache.
+ */
+template<typename ItemType>
+class SAutomationTestTreeView
+	: public STreeView< ItemType >
+{
+public:
+		typedef typename TListTypeTraits< ItemType >::NullableType NullableItemType;
+
+		typedef typename TSlateDelegates< ItemType >::FOnGetChildren FOnGetChildren;
+		typedef typename TSlateDelegates< ItemType >::FOnGenerateRow FOnGenerateRow;
+		typedef typename TSlateDelegates< ItemType >::FOnSetExpansionRecursive FOnSetExpansionRecursive;
+		typedef typename TSlateDelegates< ItemType >::FOnItemScrolledIntoView FOnItemScrolledIntoView;
+		typedef typename TSlateDelegates< NullableItemType >::FOnSelectionChanged FOnSelectionChanged;
+		typedef typename TSlateDelegates< ItemType >::FOnMouseButtonDoubleClick FOnMouseButtonDoubleClick;
+		typedef typename TSlateDelegates< ItemType >::FOnExpansionChanged FOnExpansionChanged;
+public:
+	
+		SLATE_BEGIN_ARGS( SAutomationTestTreeView )
+			: _OnGenerateRow()
+			, _OnGetChildren()
+			, _OnSetExpansionRecursive()
+			, _TreeItemsSource( static_cast< TArray<ItemType>* >(NULL) ) //@todo Slate Syntax: Initializing from NULL without a cast
+			, _ItemHeight(16)
+			, _OnContextMenuOpening()
+			, _OnMouseButtonDoubleClick()
+			, _OnSelectionChanged()
+			, _OnExpansionChanged()
+			, _SelectionMode(ESelectionMode::Multi)
+			, _ClearSelectionOnClick(true)
+			, _ExternalScrollbar()
+			{}
+
+			SLATE_EVENT( FOnGenerateRow, OnGenerateRow )
+
+			SLATE_EVENT( FOnItemScrolledIntoView, OnItemScrolledIntoView )
+
+			SLATE_EVENT( FOnGetChildren, OnGetChildren )
+
+			SLATE_EVENT( FOnSetExpansionRecursive, OnSetExpansionRecursive )
+
+			SLATE_ARGUMENT( TArray<ItemType>* , TreeItemsSource )
+
+			SLATE_ATTRIBUTE( float, ItemHeight )
+
+			SLATE_EVENT( FOnContextMenuOpening, OnContextMenuOpening )
+
+			SLATE_EVENT( FOnMouseButtonDoubleClick, OnMouseButtonDoubleClick )
+
+			SLATE_EVENT( FOnSelectionChanged, OnSelectionChanged )
+
+			SLATE_EVENT( FOnExpansionChanged, OnExpansionChanged )
+
+			SLATE_ATTRIBUTE( ESelectionMode::Type, SelectionMode )
+
+			SLATE_ARGUMENT( TSharedPtr<SHeaderRow>, HeaderRow )
+
+			SLATE_ARGUMENT ( bool, ClearSelectionOnClick )
+
+			SLATE_ARGUMENT( TSharedPtr<SScrollBar>, ExternalScrollbar )
+
+		SLATE_END_ARGS()
+
+		
+		/**
+		 * Construct this widget
+		 *
+		 * @param	InArgs	The declaration data for this widget
+		 */
+		void Construct( const FArguments& InArgs )
+		{
+			this->OnGenerateRow = InArgs._OnGenerateRow;
+			this->OnItemScrolledIntoView = InArgs._OnItemScrolledIntoView;
+			this->OnGetChildren = InArgs._OnGetChildren;
+			this->OnSetExpansionRecursive = InArgs._OnSetExpansionRecursive;
+			this->TreeItemsSource = InArgs._TreeItemsSource;
+
+			this->OnContextMenuOpening = InArgs._OnContextMenuOpening;
+			this->OnDoubleClick = InArgs._OnMouseButtonDoubleClick;
+			this->OnSelectionChanged = InArgs._OnSelectionChanged;
+			this->OnExpansionChanged = InArgs._OnExpansionChanged;
+			this->SelectionMode = InArgs._SelectionMode;
+
+			this->bClearSelectionOnClick = InArgs._ClearSelectionOnClick;
+
+			// Check for any parameters that the coder forgot to specify.
+			FString ErrorString;
+			{
+				if ( !this->OnGenerateRow.IsBound() )
+				{
+					ErrorString += TEXT("Please specify an OnGenerateRow. \n");
+				}
+
+				if ( this->TreeItemsSource == NULL )
+				{
+					ErrorString += TEXT("Please specify a TreeItemsSource. \n");
+				}
+
+				if ( !this->OnGetChildren.IsBound() )
+				{
+					ErrorString += TEXT("Please specify an OnGetChildren. \n");
+				}
+			}
+
+			if (ErrorString.Len() > 0)
+			{
+				// Let the coder know what they forgot
+				this->ChildSlot
+				.HAlign(HAlign_Center)
+				.VAlign(VAlign_Center)
+				[
+					SNew(STextBlock)
+					.Text(ErrorString)
+				];
+			}
+			else
+			{
+				// Make the TableView
+				this->ConstructChildren( 0, InArgs._ItemHeight, InArgs._HeaderRow, InArgs._ExternalScrollbar );
+			}
+		}
+
+		SAutomationTestTreeView()
+			: STreeView< ItemType >()
+		{
+		}
+
+public:
+
+	/**
+	 * Clears the internal widget cache and recreates the tree
+	 */
+	void ReCreateTreeView()
+	{
+		WidgetGenerator.Clear();
+		RequestTreeRefresh();
+	}
+};
diff --git a/Engine/Source/Developer/AutomationWindow/Private/SAutomationWindow.cpp b/Engine/Source/Developer/AutomationWindow/Private/SAutomationWindow.cpp
--- a/Engine/Source/Developer/AutomationWindow/Private/SAutomationWindow.cpp
+++ b/Engine/Source/Developer/AutomationWindow/Private/SAutomationWindow.cpp
@@ -82,6 +82,8 @@ void SAutomationWindow::Construct( const FArguments& InArgs, const IAutomationCo
 	TestPresetManager->LoadPresets();
 	bAddingTestPreset = false;
 
+	bHasChildTestSelected = false;
+
 	SessionManager = InSessionManager;
 	AutomationController = InAutomationController;
 
@@ -105,7 +107,7 @@ void SAutomationWindow::Construct( const FArguments& InArgs, const IAutomationCo
 	//make the widget for platforms
 	PlatformsHBox = SNew (SHorizontalBox);
 
-	TestTable = SNew(STreeView< TSharedPtr< IAutomationReport > >)
+	TestTable = SNew(SAutomationTestTreeView< TSharedPtr< IAutomationReport > >)
 		.SelectionMode(ESelectionMode::Multi)
 		.TreeItemsSource( &TestReports )
 		// Generates the actual widget for a tree item
@@ -235,41 +237,121 @@ void SAutomationWindow::Construct( const FArguments& InArgs, const IAutomationCo
 			+SSplitter::Slot()
 			.Value(0.33f)
 			[
-				//results panel
-				SNew( SVerticalBox )
-				+SVerticalBox::Slot()
-				.AutoHeight()
-				[
-					SNew( STextBlock )
-						.Text( LOCTEXT("AutomationTest_Results", "Automation Test Results:") )
-				]
-				+SVerticalBox::Slot()
-				.FillHeight(1.0f)
-				.Padding(0.0f, 4.0f, 0.0f, 0.0f)
+				SNew(SOverlay)
+				+SOverlay::Slot()
 				[
-					//list of results for the selected test
-					SNew(SBorder)
+					SNew(SBox)
+					.Visibility(this, &SAutomationWindow::GetTestGraphVisibility)
 					[
-						SAssignNew(LogListView, SListView<TSharedPtr<FAutomationOutputMessage> >)
-						.ItemHeight(18)
-						.ListItemsSource(&LogMessages)
-						.SelectionMode(ESelectionMode::Multi)
-						.OnGenerateRow(this, &SAutomationWindow::OnGenerateWidgetForLog)
-						.OnSelectionChanged(this, &SAutomationWindow::HandleLogListSelectionChanged)
+						//Graphical Results Panel
+						SNew( SVerticalBox )
+						+SVerticalBox::Slot()
+						.AutoHeight()
+						[
+							SNew(SHorizontalBox)
+							+SHorizontalBox::Slot()
+							.HAlign(HAlign_Left)
+							[
+								SNew(STextBlock)
+								.Text( LOCTEXT("AutomationTest_GraphicalResults", "Automation Test Graphical Results:"))
+							]
+
+							+SHorizontalBox::Slot()
+							.HAlign(HAlign_Right)
+							.AutoWidth()
+							[
+								SNew(STextBlock)
+								.Text( LOCTEXT("AutomationTest_GraphicalResults", "Display:"))
+							]
+
+							+SHorizontalBox::Slot()
+							.HAlign(HAlign_Right)
+							.AutoWidth()
+							[
+								SNew(SCheckBox)
+								.Style(FCoreStyle::Get(), "RadioButton")
+								.IsChecked(this, &SAutomationWindow::HandleResultDisplayTypeIsChecked, EAutomationGrapicalDisplayType::DisplayName)
+								.OnCheckStateChanged(this, &SAutomationWindow::HandleResultDisplayTypeStateChanged, EAutomationGrapicalDisplayType::DisplayName)
+								[
+									SNew(STextBlock)
+									.Text( LOCTEXT("AutomationTest_GraphicalResultsDisplayName", "Name"))
+								]
+							]
+
+							+SHorizontalBox::Slot()
+							.HAlign(HAlign_Right)
+							.AutoWidth()
+							[
+								SNew(SCheckBox)
+								.Style(FCoreStyle::Get(), "RadioButton")
+								.IsChecked(this, &SAutomationWindow::HandleResultDisplayTypeIsChecked, EAutomationGrapicalDisplayType::DisplayTime)
+								.OnCheckStateChanged(this, &SAutomationWindow::HandleResultDisplayTypeStateChanged, EAutomationGrapicalDisplayType::DisplayTime)
+								[
+									SNew(STextBlock)
+									.Text( LOCTEXT("AutomationTest_GraphicalResultsDisplayTime", "Time"))
+								]
+							]
+						]
+
+						+SVerticalBox::Slot()
+						.FillHeight(1.0f)
+						[
+							SNew(SBorder)
+							[
+								SNew(SScrollBox)
+								+ SScrollBox::Slot()
+								[
+									SAssignNew(GraphicalResultBox, SAutomationGraphicalResultBox, InAutomationController)
+								]
+							]
+						]
 					]
 				]
-				+SVerticalBox::Slot()
-				.AutoHeight()
-				.Padding(0.0f, 4.0f, 0.0f, 0.0f)
+
+				+SOverlay::Slot()
 				[
-					SNew(SBorder)
-						.BorderImage(FEditorStyle::GetBrush("ToolPanel.GroupBorder"))
-						.Padding(FMargin(8.0f, 6.0f))
+					SNew(SBox)
+					.Visibility(this, &SAutomationWindow::GetTestLogVisibility)
+					[
+						//results panel
+						SNew( SVerticalBox )
+						+SVerticalBox::Slot()
+						.AutoHeight()
+						[
+							SNew( STextBlock )
+							.Text( LOCTEXT("AutomationTest_Results", "Automation Test Results:") )
+						]
+
+						+SVerticalBox::Slot()
+						.FillHeight(1.0f)
+						.Padding(0.0f, 4.0f, 0.0f, 0.0f)
 						[
-							// Add the command bar
-							SAssignNew(CommandBar, SAutomationWindowCommandBar, NotificationList)
-							.OnCopyLogClicked(this, &SAutomationWindow::HandleCommandBarCopyLogClicked)
+							//list of results for the selected test
+							SNew(SBorder)
+							[
+								SAssignNew(LogListView, SListView<TSharedPtr<FAutomationOutputMessage> >)
+								.ItemHeight(18)
+								.ListItemsSource(&LogMessages)
+								.SelectionMode(ESelectionMode::Multi)
+								.OnGenerateRow(this, &SAutomationWindow::OnGenerateWidgetForLog)
+								.OnSelectionChanged(this, &SAutomationWindow::HandleLogListSelectionChanged)
+							]
 						]
+
+						+SVerticalBox::Slot()
+						.AutoHeight()
+						.Padding(0.0f, 4.0f, 0.0f, 0.0f)
+						[
+							SNew(SBorder)
+							.BorderImage(FEditorStyle::GetBrush("ToolPanel.GroupBorder"))
+							.Padding(FMargin(8.0f, 6.0f))
+							[
+								// Add the command bar
+								SAssignNew(CommandBar, SAutomationWindowCommandBar, NotificationList)
+								.OnCopyLogClicked(this, &SAutomationWindow::HandleCommandBarCopyLogClicked)
+							]
+						]
+					]
 				]
 			]
 		]
@@ -304,6 +386,22 @@ void SAutomationWindow::Construct( const FArguments& InArgs, const IAutomationCo
 }
 END_SLATE_FUNCTION_BUILD_OPTIMIZATION
 
+
+void SAutomationWindow::HandleResultDisplayTypeStateChanged( ESlateCheckBoxState::Type NewRadioState, EAutomationGrapicalDisplayType::Type NewDisplayType)
+{
+	if (NewRadioState == ESlateCheckBoxState::Checked)
+	{
+		GraphicalResultBox->SetDisplayType(NewDisplayType);
+	}
+}
+
+ESlateCheckBoxState::Type SAutomationWindow::HandleResultDisplayTypeIsChecked( EAutomationGrapicalDisplayType::Type InDisplayType ) const
+{
+	return (GraphicalResultBox->GetDisplayType() == InDisplayType)
+		? ESlateCheckBoxState::Checked
+		: ESlateCheckBoxState::Unchecked;
+}
+
 const FSlateBrush* SAutomationWindow::GetTestBackgroundBorderImage() const
 {
 	switch(TestBackgroundType)
@@ -405,6 +503,16 @@ TSharedRef< SWidget > SAutomationWindow::MakeAutomationWindowToolBar( const TSha
 				ToolbarBuilder.AddToolBarButton( FAutomationWindowCommands::Get().VisualCommandlet );
 			}
 			ToolbarBuilder.EndSection();
+			ToolbarBuilder.BeginSection("GroupFlags");
+			{
+				ToolbarBuilder.AddComboButton(
+					FUIAction(),
+					FOnGetContent::CreateStatic( &SAutomationWindow::GenerateGroupOptionsMenuContent, InAutomationWindow ),
+					LOCTEXT( "GroupOptions_Label", "Device Groups" ),
+					LOCTEXT( "GroupOptionsToolTip", "Device Group Options" ),
+					FSlateIcon(FEditorStyle::GetStyleSetName(), "AutomationWindow.GroupSettings"));
+			}
+			ToolbarBuilder.EndSection();
 			ToolbarBuilder.BeginSection("Search");
 			{
 				ToolbarBuilder.AddWidget( Searchbox );
@@ -747,6 +855,73 @@ TSharedRef<SWidget> SAutomationWindow::GeneratePresetComboItem(TSharedPtr<FAutom
 		.Text( InItem->GetPresetName() );
 }
 
+TSharedRef< SWidget > SAutomationWindow::GenerateGroupOptionsMenuContent( TWeakPtr<class SAutomationWindow> InAutomationWindow )
+{
+	TSharedPtr<SAutomationWindow> AutomationWindow(InAutomationWindow.Pin());
+	if( AutomationWindow.IsValid() )
+	{
+		return AutomationWindow->GenerateGroupOptionsMenuContent();
+	}
+
+	//Return empty menu
+	FMenuBuilder MenuBuilder( true, NULL );
+	MenuBuilder.BeginSection("AutomationWindowGroupOptions", LOCTEXT("GroupOptions", "Device Group Options"));
+	MenuBuilder.EndSection();
+	return MenuBuilder.MakeWidget();
+}
+
+TSharedRef< SWidget > SAutomationWindow::GenerateGroupOptionsMenuContent( )
+{
+	const bool bShouldCloseWindowAfterMenuSelection = true;
+	FMenuBuilder MenuBuilder( bShouldCloseWindowAfterMenuSelection, AutomationWindowActions );
+	const uint32 NumFlags = EAutomationDeviceGroupTypes::Max;
+	TSharedPtr<SWidget> FlagWidgets[NumFlags];
+	for( int32 i=0; i<NumFlags; i++ )
+	{
+		FlagWidgets[i] = 
+			SNew(SCheckBox)
+			.IsChecked(this, &SAutomationWindow::IsDeviceGroupCheckBoxIsChecked, i)
+			.OnCheckStateChanged(this, &SAutomationWindow::HandleDeviceGroupCheckStateChanged, i)
+			.Padding(FMargin(4.0f, 0.0f))
+			.ToolTipText(EAutomationDeviceGroupTypes::ToDescription((EAutomationDeviceGroupTypes::Type)i))
+			.IsEnabled( this, &SAutomationWindow::IsAutomationControllerIdle )
+			.Content()
+			[
+				SNew(STextBlock)
+				.Text(EAutomationDeviceGroupTypes::ToName((EAutomationDeviceGroupTypes::Type)i))
+			];
+	}
+
+	MenuBuilder.BeginSection("AutomationWindowGroupDevices", LOCTEXT("GroupOptions", "Group Types"));
+	{
+		for( int32 i=0; i<NumFlags;i++ )
+		{
+			MenuBuilder.AddWidget(FlagWidgets[i].ToSharedRef(),FText::GetEmpty());
+		}
+	}
+
+	return MenuBuilder.MakeWidget();
+}
+
+/** Returns if full size screen shots are enabled */
+ESlateCheckBoxState::Type SAutomationWindow::IsDeviceGroupCheckBoxIsChecked(const int32 DeviceGroupFlag) const
+{
+	return AutomationController->IsDeviceGroupFlagSet((EAutomationDeviceGroupTypes::Type)DeviceGroupFlag) ? ESlateCheckBoxState::Checked : ESlateCheckBoxState::Unchecked;
+}
+/** Toggles if we are collecting full size screenshots */
+void SAutomationWindow::HandleDeviceGroupCheckStateChanged(ESlateCheckBoxState::Type CheckBoxState, const int32 DeviceGroupFlag)
+{
+	//Update the device groups
+	AutomationController->ToggleDeviceGroupFlag((EAutomationDeviceGroupTypes::Type)DeviceGroupFlag);
+	AutomationController->UpdateDeviceGroups();
+	
+	//Update header
+	RebuildPlatformIcons();
+
+	//Need to force the tree to do a full refresh here because the reports have changed but the tree will keep using cached data.
+	TestTable->ReCreateTreeView();
+}
+
 TSharedRef< SWidget > SAutomationWindow::GenerateTestsOptionsMenuContent( TWeakPtr<class SAutomationWindow> InAutomationWindow )
 {
 	TSharedPtr<SAutomationWindow> AutomationWindow(InAutomationWindow.Pin());
@@ -848,7 +1023,7 @@ void SAutomationWindow::HandleFullSizeScreenshotsBoxCheckStateChanged(ESlateChec
 
 ESlateCheckBoxState::Type SAutomationWindow::IsEnableScreenshotsCheckBoxChecked() const
 {
-	return AutomationController->AreScreenshotsEnabled() ? ESlateCheckBoxState::Checked : ESlateCheckBoxState::Unchecked;
+	return AutomationController->IsScreenshotAllowed() ? ESlateCheckBoxState::Checked : ESlateCheckBoxState::Unchecked;
 }
 
 void SAutomationWindow::HandleEnableScreenshotsBoxCheckStateChanged(ESlateCheckBoxState::Type CheckBoxState)
@@ -858,7 +1033,7 @@ void SAutomationWindow::HandleEnableScreenshotsBoxCheckStateChanged(ESlateCheckB
 
 bool SAutomationWindow::IsFullSizeScreenshotsOptionEnabled() const
 {
-	return AutomationController->AreScreenshotsEnabled() && IsAutomationControllerIdle();
+	return AutomationController->IsScreenshotAllowed() && IsAutomationControllerIdle();
 }
 
 // Only valid in the editor
@@ -902,8 +1077,14 @@ void SAutomationWindow::OnTestSelectionChanged(TSharedPtr<IAutomationReport> Sel
 	//empty the previous log
 	LogMessages.Empty();
 
+	bHasChildTestSelected = false;
+
 	if (Selection.IsValid())
 	{
+		if( Selection->GetTotalNumChildren() == 0 )
+		{
+			bHasChildTestSelected = true;
+		}
 		//accumulate results for each device cluster that supports the test
 		int32 NumClusters = AutomationController->GetNumDeviceClusters();
 		for (int32 ClusterIndex = 0; ClusterIndex < NumClusters; ++ClusterIndex)
@@ -911,7 +1092,7 @@ void SAutomationWindow::OnTestSelectionChanged(TSharedPtr<IAutomationReport> Sel
 			//no sense displaying device name if only one is available
 			if (NumClusters > 1)
 			{
-				FString DeviceTypeName = AutomationController->GetDeviceTypeName(ClusterIndex);
+				FString DeviceTypeName = AutomationController->GetClusterGroupName(ClusterIndex) + TEXT("  -  ") + Selection->GetGameInstanceName(ClusterIndex);
 				LogMessages.Add(MakeShareable(new FAutomationOutputMessage(DeviceTypeName, TEXT("Automation.Header"))));
 			}
 
@@ -949,9 +1130,23 @@ void SAutomationWindow::OnTestSelectionChanged(TSharedPtr<IAutomationReport> Sel
 				LogMessages.Add(MakeShareable(new FAutomationOutputMessage(TEXT(""), TEXT("Log.Normal"))));
 			}
 		}
-		//rebuild UI
-		LogListView->RequestListRefresh();
 	}
+
+	//rebuild UI
+	LogListView->RequestListRefresh();
+}
+
+
+EVisibility SAutomationWindow::GetTestLogVisibility( ) const
+{
+	return (GetTestGraphVisibility() == EVisibility::Visible) ? EVisibility::Hidden : EVisibility::Visible;
+}
+
+
+EVisibility SAutomationWindow::GetTestGraphVisibility( ) const
+{
+	//Show the graphical window if we don't have a child test selected and we have results to view
+	return (!bHasChildTestSelected && GraphicalResultBox->HasResults()) ? EVisibility::Visible : EVisibility::Hidden;
 }
 
 
@@ -1013,12 +1208,20 @@ END_SLATE_FUNCTION_BUILD_OPTIMIZATION
 FText SAutomationWindow::CreateDeviceTooltip(int32 ClusterIndex)
 {
 	FTextBuilder ReportBuilder;
+
+	const int32 NumClusters = AutomationController->GetNumDeviceClusters();
+	if( NumClusters > 1 )
+	{
+		ReportBuilder.AppendLine(LOCTEXT("ToolTipClusterName", "Cluster Name:"));
+		ReportBuilder.AppendLine(AutomationController->GetClusterGroupName(ClusterIndex));
+	}
+
 	ReportBuilder.AppendLine(LOCTEXT("ToolTipGameInstances", "Game Instances:"));
 
 	int32 NumDevices = AutomationController->GetNumDevicesInCluster( ClusterIndex );
 	for ( int32 DeviceIndex = 0; DeviceIndex < NumDevices; ++DeviceIndex )
 	{
-		ReportBuilder.AppendLine(AutomationController->GetGameInstanceName(ClusterIndex, DeviceIndex));
+		ReportBuilder.AppendLine(AutomationController->GetGameInstanceName(ClusterIndex, DeviceIndex).LeftPad(2));
 	}
 
 	return ReportBuilder.ToText();
@@ -1220,6 +1423,9 @@ FReply SAutomationWindow::RunTests()
 	LogMessages.Empty();
 	LogListView->RequestListRefresh();
 
+	//Clear old results
+	GraphicalResultBox->ClearResults();
+
 	return FReply::Handled();
 }
 
diff --git a/Engine/Source/Developer/AutomationWindow/Private/SAutomationWindow.h b/Engine/Source/Developer/AutomationWindow/Private/SAutomationWindow.h
--- a/Engine/Source/Developer/AutomationWindow/Private/SAutomationWindow.h
+++ b/Engine/Source/Developer/AutomationWindow/Private/SAutomationWindow.h
@@ -157,6 +157,15 @@ private:
 	static TSharedRef< SWidget >GenerateTestsOptionsMenuContent( TWeakPtr<class SAutomationWindow> InAutomationWindow );
 	TSharedRef< SWidget > GenerateTestsOptionsMenuContent( );
 
+	
+	/**
+	 * Static: Creates the group flag options menu widget
+	 *
+	 * @return	New widget
+	 */
+	static TSharedRef< SWidget >GenerateGroupOptionsMenuContent( TWeakPtr<class SAutomationWindow> InAutomationWindow );
+	TSharedRef< SWidget > GenerateGroupOptionsMenuContent( );
+
 	/**
 	 * Creates a combo item for the preset list
 	 *
@@ -242,6 +251,11 @@ private:
 
 	/** Returns if the full size screenshots option is enabled */
 	bool IsFullSizeScreenshotsOptionEnabled() const;
+
+	/** Returns if a device group is enabled */
+	ESlateCheckBoxState::Type IsDeviceGroupCheckBoxIsChecked(const int32 DeviceGroupFlag) const;
+	/** Toggles a device group flag */
+	void HandleDeviceGroupCheckStateChanged(ESlateCheckBoxState::Type CheckBoxState, const int32 DeviceGroupFlag);
 	
 	/** Sets the number of times to repeat the tests */
 	void OnChangeRepeatCount(int32 InNewValue);
@@ -418,6 +432,26 @@ private:
 	*/
 	void HandleLogListSelectionChanged( TSharedPtr<FAutomationOutputMessage> InItem, ESelectInfo::Type SelectInfo );
 
+	/**
+	* Gets the visibility of the test log window
+	*/
+	EVisibility GetTestLogVisibility( ) const;
+
+	/**
+	* Gets the visibility of the graphical test result window
+	*/
+	EVisibility GetTestGraphVisibility( ) const;
+
+	/**
+	* Handles changing the display type for the graphical results window
+	*/
+	void HandleResultDisplayTypeStateChanged( ESlateCheckBoxState::Type NewRadioState, EAutomationGrapicalDisplayType::Type NewDisplayType );
+
+	/**
+	* Handles checking if a display type is active for the graphical results window
+	*/
+	ESlateCheckBoxState::Type HandleResultDisplayTypeIsChecked( EAutomationGrapicalDisplayType::Type InDisplayType ) const;
+
 	/** 
 	 * Gets the visibility for the throbber
 	 *
@@ -449,7 +483,7 @@ private:
 	TSharedPtr< SCheckBox > HeaderCheckbox;
 
 	// The list of all valid tests.
-	TSharedPtr< STreeView< TSharedPtr< IAutomationReport > > > TestTable;
+	TSharedPtr< SAutomationTestTreeView< TSharedPtr< IAutomationReport > > > TestTable;
 	
 	// Widget for header platform icons.
 	TSharedPtr< SHorizontalBox > PlatformsHBox;
@@ -463,6 +497,9 @@ private:
 	// Holds the widget to display log messages
 	TSharedPtr<SListView<TSharedPtr<FAutomationOutputMessage> > > LogListView;
 
+	// Holds the widget to display a graph of the results
+	TSharedPtr< SAutomationGraphicalResultBox > GraphicalResultBox;
+
 	// Holds the collection of log messages.
 	TArray<TSharedPtr<FAutomationOutputMessage> > LogMessages;
 
@@ -485,6 +522,9 @@ private:
 	// Flag to acknowledge if the window is awaiting tests to display
 	bool bIsRequestingTests;
 
+	// Flag to tell if we have a child test selected in the test tree */
+	bool bHasChildTestSelected;
+
 	// Which type of window style to use for the test background
 	EAutomationTestBackgroundStyle::Type TestBackgroundType;
 
diff --git a/Engine/Source/Editor/EditorStyle/Private/SlateEditorStyle.cpp b/Engine/Source/Editor/EditorStyle/Private/SlateEditorStyle.cpp
--- a/Engine/Source/Editor/EditorStyle/Private/SlateEditorStyle.cpp
+++ b/Engine/Source/Editor/EditorStyle/Private/SlateEditorStyle.cpp
@@ -4930,6 +4930,10 @@ void FSlateEditorStyle::FStyle::SetupAutomationStyles()
 		Set( "Automation.Error"  , FTextBlockStyle(NormalText)
 			.SetFont( TTF_FONT( "Fonts/DroidSansMono", 9 ) )
 			.SetColorAndOpacity(FLinearColor(FColor(0xffff0000))) );
+
+		Set( "Automation.ReportHeader" , FTextBlockStyle(NormalText)
+			.SetFont( TTF_FONT( "Fonts/DroidSansMono", 10 ) )
+			.SetColorAndOpacity(FLinearColor(FColor(0xffffffff))) );
 		
 		//state of individual tests
 		Set( "Automation.Success", new IMAGE_BRUSH( "Automation/Success", Icon16x16 ) );
@@ -4967,6 +4971,10 @@ void FSlateEditorStyle::FStyle::SetupAutomationStyles()
 		Set( "AutomationWindow.DeveloperDirectoryContent.Small", new IMAGE_BRUSH( "Automation/DeveloperDirectoryContent", Icon20x20) );
 		Set( "AutomationWindow.VisualCommandlet.Small", new IMAGE_BRUSH( "Automation/VisualCommandlet", Icon20x20) );
 
+		//device group settings
+		Set( "AutomationWindow.GroupSettings", new IMAGE_BRUSH( "Automation/Groups", Icon40x40) );
+		Set( "AutomationWindow.GroupSettings.Small", new IMAGE_BRUSH( "Automation/Groups", Icon20x20) );
+
 		//test preset icons
 		Set( "AutomationWindow.PresetNew", new IMAGE_BRUSH( "Icons/icon_add_40x", Icon16x16 ) );
 		Set( "AutomationWindow.PresetSave", new IMAGE_BRUSH( "Icons/icon_file_save_16px", Icon16x16 ) );
diff --git a/Engine/Source/Editor/StaticMeshEditor/Private/StaticMeshAutomationTests.cpp b/Engine/Source/Editor/StaticMeshEditor/Private/StaticMeshAutomationTests.cpp
--- a/Engine/Source/Editor/StaticMeshEditor/Private/StaticMeshAutomationTests.cpp
+++ b/Engine/Source/Editor/StaticMeshEditor/Private/StaticMeshAutomationTests.cpp
@@ -203,7 +203,7 @@ bool FStaticMeshEditorTest::RunTest(const FString& Parameters)
 
 		//Wait for the window to load and then take the initial screen shot		
 		ADD_LATENT_AUTOMATION_COMMAND(FWaitLatentCommand(0.5f));
-		if( FAutomationTestFramework::GetInstance().AreScreenshotsEnabled() )
+		if( FAutomationTestFramework::GetInstance().IsScreenshotAllowed() )
 		{
 			AutomationCommon::GetScreenshotPath(TestName, WindowParameters.ScreenshotName, false);
 			ADD_LATENT_AUTOMATION_COMMAND(FTakeEditorScreenshotCommand(WindowParameters));
@@ -218,7 +218,7 @@ bool FStaticMeshEditorTest::RunTest(const FString& Parameters)
 			ADD_LATENT_AUTOMATION_COMMAND(EditorViewButtonHelper::FPerformStaticMeshFlagToggle(AutomationParameters));
 			ADD_LATENT_AUTOMATION_COMMAND(FWaitLatentCommand(0.1f));
 
-			if( FAutomationTestFramework::GetInstance().AreScreenshotsEnabled() )
+			if( FAutomationTestFramework::GetInstance().IsScreenshotAllowed() )
 			{
 				//Take the screenshot
 				AutomationCommon::GetScreenshotPath(TestName, WindowParameters.ScreenshotName, false);
diff --git a/Engine/Source/Editor/UnrealEd/Private/Tests/EditorAutomationTests.cpp b/Engine/Source/Editor/UnrealEd/Private/Tests/EditorAutomationTests.cpp
--- a/Engine/Source/Editor/UnrealEd/Private/Tests/EditorAutomationTests.cpp
+++ b/Engine/Source/Editor/UnrealEd/Private/Tests/EditorAutomationTests.cpp
@@ -874,7 +874,7 @@ bool FLoadAllMapsInEditorTest::RunTest(const FString& Parameters)
 {
 	FString MapName = Parameters;
 
-	const bool bTakeScreenshots = FAutomationTestFramework::GetInstance().AreScreenshotsEnabled();
+	const bool bTakeScreenshots = FAutomationTestFramework::GetInstance().IsScreenshotAllowed();
 	if( bTakeScreenshots )
 	{
 		//Find the main editor window
diff --git a/Engine/Source/Runtime/AutomationMessages/Classes/AutomationWorkerMessages.h b/Engine/Source/Runtime/AutomationMessages/Classes/AutomationWorkerMessages.h
--- a/Engine/Source/Runtime/AutomationMessages/Classes/AutomationWorkerMessages.h
+++ b/Engine/Source/Runtime/AutomationMessages/Classes/AutomationWorkerMessages.h
@@ -97,6 +97,42 @@ struct FAutomationWorkerFindWorkersResponse
 	UPROPERTY()
 	FString Platform;
 
+	/**
+	 * Holds the name of the operating system version.
+	 */
+	UPROPERTY()
+	FString OSVersionName;
+
+	/**
+	 * Holds the name of the device model.
+	 */
+	UPROPERTY()
+	FString ModelName;
+
+	/**
+	 * Holds the name of the GPU.
+	 */
+	UPROPERTY()
+	FString GPUName;
+
+	/**
+	 * Holds the name of the CPU model.
+	 */
+	UPROPERTY()
+	FString CPUModelName;
+
+	/**
+	 * Holds the amount of RAM this device has in gigabytes.
+	 */
+	UPROPERTY()
+	uint32 RAMInGB;
+
+	/**
+	 * Holds the name of the current render mode.
+	 */
+	UPROPERTY()
+	FString RenderModeName;
+
 	/**
 	 * Holds the worker's application session identifier.
 	 */
@@ -108,16 +144,6 @@ struct FAutomationWorkerFindWorkersResponse
 	 * Default constructor.
 	 */
 	FAutomationWorkerFindWorkersResponse( ) { }
-
-	/**
-	 * Creates and initializes a new instance.
-	 */
-	FAutomationWorkerFindWorkersResponse( const FString& InDeviceName, const FString& InInstanceName, const FString& InPlatform, const FGuid& InSessionId )
-		: DeviceName(InDeviceName)
-		, InstanceName(InInstanceName)
-		, Platform(InPlatform)
-		, SessionId(InSessionId)
-	{ }
 };
 
 template<>
diff --git a/Engine/Source/Runtime/AutomationWorker/AutomationWorker.Build.cs b/Engine/Source/Runtime/AutomationWorker/AutomationWorker.Build.cs
--- a/Engine/Source/Runtime/AutomationWorker/AutomationWorker.Build.cs
+++ b/Engine/Source/Runtime/AutomationWorker/AutomationWorker.Build.cs
@@ -38,6 +38,7 @@ namespace UnrealBuildTool.Rules
 			if (UEBuildConfiguration.bCompileAgainstEngine)
 			{
 				PrivateDependencyModuleNames.Add("Engine");
+				PrivateDependencyModuleNames.Add("RHI");
 			}
 		}
 	}
diff --git a/Engine/Source/Runtime/AutomationWorker/Private/AutomationWorkerModule.cpp b/Engine/Source/Runtime/AutomationWorker/Private/AutomationWorkerModule.cpp
--- a/Engine/Source/Runtime/AutomationWorker/Private/AutomationWorkerModule.cpp
+++ b/Engine/Source/Runtime/AutomationWorker/Private/AutomationWorkerModule.cpp
@@ -249,7 +249,30 @@ void FAutomationWorkerModule::HandleFindWorkersMessage( const FAutomationWorkerF
 
 	if ((Message.SessionId == FApp::GetSessionId()) && (Message.Changelist == 10000))
 	{
-		MessageEndpoint->Send(new FAutomationWorkerFindWorkersResponse(FPlatformProcess::ComputerName(), InstanceName, FPlatformProperties::PlatformName(), Message.SessionId), Context->GetSender());
+		FAutomationWorkerFindWorkersResponse* Response = new FAutomationWorkerFindWorkersResponse();
+
+		FString OSMajorVersionString, OSSubVersionString;
+		FPlatformMisc::GetOSVersions( OSMajorVersionString, OSSubVersionString );
+
+		FString OSVersionString = OSMajorVersionString + TEXT(" ") + OSSubVersionString;
+		FString CPUModelString = FPlatformMisc::GetCPUBrand().Trim();
+
+		Response->DeviceName = FPlatformProcess::ComputerName();
+		Response->InstanceName = InstanceName;
+		Response->Platform = FPlatformProperties::PlatformName();
+		Response->SessionId = Message.SessionId;
+		Response->OSVersionName = OSVersionString;
+		Response->ModelName = FPlatformMisc::GetDefaultDeviceProfileName();
+		Response->GPUName = FPlatformMisc::GetPrimaryGPUBrand();
+		Response->CPUModelName = CPUModelString;
+		Response->RAMInGB = FPlatformMemory::GetPhysicalGBRam();
+#if WITH_ENGINE
+		Response->RenderModeName = AutomationCommon::GetRenderDetailsString();
+#else
+		Response->RenderModeName = TEXT("Unknown");
+#endif
+
+		MessageEndpoint->Send(Response, Context->GetSender());
 	}
 }
 
@@ -290,7 +313,7 @@ void FAutomationWorkerModule::HandleRequestTestsMessage( const FAutomationWorker
 #if WITH_ENGINE
 void FAutomationWorkerModule::HandleScreenShotCaptured( int32 Width, int32 Height, const TArray<FColor>& Bitmap, const FString& ScreenShotName )
 {
-	if( FAutomationTestFramework::GetInstance().AreScreenshotsEnabled() )
+	if( FAutomationTestFramework::GetInstance().IsScreenshotAllowed() )
 	{
 		int32 NewHeight = Height;
 		int32 NewWidth = Width;
diff --git a/Engine/Source/Runtime/AutomationWorker/Private/AutomationWorkerPrivatePCH.h b/Engine/Source/Runtime/AutomationWorker/Private/AutomationWorkerPrivatePCH.h
--- a/Engine/Source/Runtime/AutomationWorker/Private/AutomationWorkerPrivatePCH.h
+++ b/Engine/Source/Runtime/AutomationWorker/Private/AutomationWorkerPrivatePCH.h
@@ -19,6 +19,7 @@
 #if WITH_ENGINE
 	#include "Engine.h"
 	#include "ImageUtils.h"
+	#include "AutomationCommon.h"
 #endif
 
 /* Private includes
diff --git a/Engine/Source/Runtime/Core/Private/Misc/AutomationTest.cpp b/Engine/Source/Runtime/Core/Private/Misc/AutomationTest.cpp
--- a/Engine/Source/Runtime/Core/Private/Misc/AutomationTest.cpp
+++ b/Engine/Source/Runtime/Core/Private/Misc/AutomationTest.cpp
@@ -447,7 +447,7 @@ void FAutomationTestFramework::SetScreenshotOptions( const bool bInScreenshotsEn
 	bUseFullSizeScreenShots = bInUseFullSizeScreenshots;
 }
 
-bool FAutomationTestFramework::AreScreenshotsEnabled() const
+bool FAutomationTestFramework::IsScreenshotAllowed() const
 {
 	return bScreenshotsEnabled;
 }
diff --git a/Engine/Source/Runtime/Core/Public/Misc/AutomationTest.h b/Engine/Source/Runtime/Core/Public/Misc/AutomationTest.h
--- a/Engine/Source/Runtime/Core/Public/Misc/AutomationTest.h
+++ b/Engine/Source/Runtime/Core/Public/Misc/AutomationTest.h
@@ -511,9 +511,9 @@ public:
 	void SetScreenshotOptions( const bool bInScreenshotsEnabled, const bool bInUseFullSizeScreenshots );
 
 	/**
-	 * Gets if screenshots are enabled
+	 * Gets if screenshots are allowed
 	 */
-	bool AreScreenshotsEnabled() const;
+	bool IsScreenshotAllowed() const;
 
 	/**
 	 * Gets if we are using fulll size screenshots
diff --git a/Engine/Source/Runtime/Engine/Private/Tests/EngineAutomationTests.cpp b/Engine/Source/Runtime/Engine/Private/Tests/EngineAutomationTests.cpp
--- a/Engine/Source/Runtime/Engine/Private/Tests/EngineAutomationTests.cpp
+++ b/Engine/Source/Runtime/Engine/Private/Tests/EngineAutomationTests.cpp
@@ -179,7 +179,7 @@ bool FLoadAllMapsInGameTest::RunTest(const FString& Parameters)
 	//Open the map
 	GEngine->Exec(GetSimpleEngineAutomationTestGameWorld(GetTestFlags()), *FString::Printf(TEXT("Open %s"), *MapName));
 
-	if( FAutomationTestFramework::GetInstance().AreScreenshotsEnabled() )
+	if( FAutomationTestFramework::GetInstance().IsScreenshotAllowed() )
 	{
 		//Generate the screen shot name and path
 		FString ScreenshotFileName;
diff --git a/Engine/Source/Runtime/Engine/Public/Tests/AutomationCommon.h b/Engine/Source/Runtime/Engine/Public/Tests/AutomationCommon.h
--- a/Engine/Source/Runtime/Engine/Public/Tests/AutomationCommon.h
+++ b/Engine/Source/Runtime/Engine/Public/Tests/AutomationCommon.h
@@ -52,6 +52,12 @@ namespace AutomationCommon
 			HardwareDetailsString = ( HardwareDetailsString + TEXT( "_" ) ) + TEXT("STEREO");
 		}
 
+		if( HardwareDetailsString.Len() > 0 )
+		{
+			//Get rid of the leading "_"
+			HardwareDetailsString = HardwareDetailsString.RightChop(1);
+		}
+
 		return HardwareDetailsString;
 	}
 
@@ -62,7 +68,7 @@ namespace AutomationCommon
 
 		if( bIncludeHardwareDetails )
 		{
-			PathName = PathName + GetRenderDetailsString();
+			PathName = PathName + TEXT("_") + GetRenderDetailsString();
 		}
 
 		FPaths::MakePathRelativeTo(PathName, *FPaths::RootDir());
