Fix build errors under Clang.

diff --git a/Engine/Source/Developer/AutomationWindow/Private/SAutomationTestTreeView.h b/Engine/Source/Developer/AutomationWindow/Private/SAutomationTestTreeView.h
--- a/Engine/Source/Developer/AutomationWindow/Private/SAutomationTestTreeView.h
+++ b/Engine/Source/Developer/AutomationWindow/Private/SAutomationTestTreeView.h
@@ -141,7 +141,7 @@ public:
 	 */
 	void ReCreateTreeView()
 	{
-		WidgetGenerator.Clear();
-		RequestTreeRefresh();
+		this->WidgetGenerator.Clear();
+		this->RequestTreeRefresh();
 	}
 };
