Fixed the usage of the automation window filters.

Selecting Warnings only will now display only warnings, and selecting errors only, will display errors only.
Selecting both Errors and Warnings will display both

TTP 329255 - UFE: Automation Tab: Errors Filter will show warnings as well.

diff --git a/Engine/Source/Developer/AutomationWindow/Private/AutomationFilter.h b/Engine/Source/Developer/AutomationWindow/Private/AutomationFilter.h
--- a/Engine/Source/Developer/AutomationWindow/Private/AutomationFilter.h
+++ b/Engine/Source/Developer/AutomationWindow/Private/AutomationFilter.h
@@ -33,9 +33,18 @@ public:
 			FilterPassed = false;
 		}
 
-		if( ShowErrors || ShowWarnings )
+		if (ShowWarnings && ShowErrors)
 		{
-		 	FilterPassed = InReport->HasErrors() || InReport->HasWarnings();
+			FilterPassed = InReport->HasWarnings() || InReport->HasErrors();
+		}
+		else if (!ShowWarnings && ShowErrors)
+		{
+			FilterPassed = InReport->HasErrors();
+		}
+		else if (ShowWarnings && !ShowErrors)
+		{
+			// Do not show report as a warning if it should be highlighted as an Error!
+			FilterPassed = InReport->HasWarnings() && !InReport->HasErrors();
 		}
 
 		return FilterPassed;
