Updating the tooltips for tests to show which machine their operations were on.

TTP 314117 - UFE: Hovering over the "In Progress" icon for an automation test which is running does not show the machine on which it is being run

diff --git a/Engine/Source/Developer/AutomationWindow/Private/SAutomationTestItem.cpp b/Engine/Source/Developer/AutomationWindow/Private/SAutomationTestItem.cpp
--- a/Engine/Source/Developer/AutomationWindow/Private/SAutomationTestItem.cpp
+++ b/Engine/Source/Developer/AutomationWindow/Private/SAutomationTestItem.cpp
@@ -221,10 +221,6 @@ FText SAutomationTestItem::GetTestToolTip( int32 ClusterIndex ) const
 	{
 		ToolTip = LOCTEXT("TestToolTipNotRun", "Not Run");
 	}
-	else if( TestState == EAutomationState::InProcess )
-	{
-		ToolTip = LOCTEXT("TestToolTipInProgress", "In progress");
-	}
 	else if( TestState == EAutomationState::NotEnoughParticipants )
 	{
 		ToolTip = LOCTEXT("ToolTipNotEnoughParticipants", "This test could not be completed as there were not enough participants.");
@@ -234,13 +230,17 @@ FText SAutomationTestItem::GetTestToolTip( int32 ClusterIndex ) const
 		FFormatNamedArguments Args;
 		Args.Add(TEXT("GameName"), FText::FromString(TestStatus->GetGameInstanceName(ClusterIndex)));
 
-		if (TestState == EAutomationState::Success)
+		if (TestState == EAutomationState::InProcess)
+		{
+			ToolTip = FText::Format(LOCTEXT("TestToolTipInProgress", "In progress on: {GameName}"), Args);
+		}
+		else if (TestState == EAutomationState::Success)
 		{
-			ToolTip = LOCTEXT("TestToolTipComplete", "Completed on: ");
+			ToolTip = FText::Format(LOCTEXT("TestToolTipComplete", "Completed on: {GameName}"), Args);
 		}
 		else
 		{
-			ToolTip = LOCTEXT("TestToolTipFailed", "Failed on: ");
+			ToolTip = FText::Format(LOCTEXT("TestToolTipFailed", "Failed on: {GameName}"), Args);
 		}
 	}
 	return ToolTip;
