Automation test history added. Off by default.

- User can now toggle Track history which will monitor subsequent tests and store a log in Saved/Automation/Logs.
- We track and show up to 5 by default, this is configurable between 1 and 10.
- We store 10 logs by default, non-configurable, so that the user can configure the number to show without deleting logs.

Note, Placeholder icon has been added

diff --git a/Engine/Config/BaseEngine.ini b/Engine/Config/BaseEngine.ini
--- a/Engine/Config/BaseEngine.ini
+++ b/Engine/Config/BaseEngine.ini
@@ -2176,4 +2176,8 @@ TilesStreamingTimeThreshold=1.0
 
 [/Script/Engine.AISystemBase]
 AISystemModuleName=AIModule
-AISystemClassName=/Script/AIModule.AISystem
\ No newline at end of file
+AISystemClassName=/Script/AIModule.AISystem
+
+[AutomationController.History]
+bTrackHistory=false
+NumberOfHistoryItemsTracked=5
\ No newline at end of file
diff --git a/Engine/Source/Developer/AutomationController/Private/AutomationControllerManager.h b/Engine/Source/Developer/AutomationController/Private/AutomationControllerManager.h
--- a/Engine/Source/Developer/AutomationController/Private/AutomationControllerManager.h
+++ b/Engine/Source/Developer/AutomationController/Private/AutomationControllerManager.h
@@ -201,6 +201,12 @@ public:
 	{
 		TestsCompleteDelegate = NewCallback;
 	}
+
+	virtual void TrackReportHistory(const bool bShouldTrack, const int32 NumReportsToTrack) override;
+
+	virtual const bool IsTrackingHistory() const override;
+
+	virtual const int32 GetNumberHistoryItemsTracking() const override;
 	// End IAutomationController Interface
 
 
@@ -381,6 +387,13 @@ private:
 
 	/** If we should request full screen screen shots */
 	bool bRequestFullScreenScreenshots;
+
+	/** If we should track any test history for the next run */
+	bool bTrackHistory;
+
+	/** The number of history items we wish to track */
+	int32 NumberOfHistoryItemsTracked;
+
 private:
 
 	// Holds a delegate that is invoked when the controller shuts down.
diff --git a/Engine/Source/Developer/AutomationController/Private/AutomationControllerManger.cpp b/Engine/Source/Developer/AutomationController/Private/AutomationControllerManger.cpp
--- a/Engine/Source/Developer/AutomationController/Private/AutomationControllerManger.cpp
+++ b/Engine/Source/Developer/AutomationController/Private/AutomationControllerManger.cpp
@@ -2,6 +2,13 @@
 
 #include "AutomationControllerPrivatePCH.h"
 
+
+namespace AutomationControllerConstants
+{
+	const FString HistoryConfigSectionName = TEXT("AutomationController.History");
+}
+
+
 void FAutomationControllerManager::RequestAvailableWorkers( const FGuid& SessionId )
 {
 	//invalidate previous tests
@@ -120,6 +127,14 @@ void FAutomationControllerManager::Init()
 	bTestResultsAvailable = false;
 	bScreenshotsEnabled = true;
 	bRequestFullScreenScreenshots = false;
+
+	// Update the ini with the settings
+	bTrackHistory = false;
+	GConfig->GetBool(*AutomationControllerConstants::HistoryConfigSectionName, TEXT("bTrackHistory"), bTrackHistory, GEngineIni);
+
+	// Default num of items to track
+	NumberOfHistoryItemsTracked = 5;
+	GConfig->GetInt(*AutomationControllerConstants::HistoryConfigSectionName, TEXT("NumberOfHistoryItemsTracked"), NumberOfHistoryItemsTracked, GEngineIni);
 }
 
 
@@ -345,6 +360,9 @@ void FAutomationControllerManager::SetTestNames( const FMessageAddress& Automati
 	if (RefreshTestResponses == DeviceClusterManager.GetNumClusters())
 	{
 		TestsRefreshedDelegate.ExecuteIfBound();
+
+		// Update the tests with tracking details
+		ReportManager.TrackHistory(bTrackHistory, NumberOfHistoryItemsTracked);
 	}
 }
 
@@ -681,4 +699,26 @@ void FAutomationControllerManager::UpdateDeviceGroups( )
 	// Update the reports in case the number of clusters changed
 	int32 NumOfClusters = DeviceClusterManager.GetNumClusters();
 	ReportManager.ClustersUpdated(NumOfClusters);
+}
+
+void FAutomationControllerManager::TrackReportHistory(const bool bShouldTrack, const int32 NumReportsToTrack)
+{
+	bTrackHistory = bShouldTrack;
+	NumberOfHistoryItemsTracked = NumReportsToTrack;
+
+	// Update the ini with the settings
+	GConfig->SetBool(*AutomationControllerConstants::HistoryConfigSectionName, TEXT("bTrackHistory"), bTrackHistory, GEngineIni);
+	GConfig->SetInt(*AutomationControllerConstants::HistoryConfigSectionName, TEXT("NumberOfHistoryItemsTracked"), NumberOfHistoryItemsTracked, GEngineIni);
+
+	ReportManager.TrackHistory(bTrackHistory, NumberOfHistoryItemsTracked);
+}
+
+const bool FAutomationControllerManager::IsTrackingHistory() const
+{
+	return bTrackHistory;
+}
+
+const int32 FAutomationControllerManager::GetNumberHistoryItemsTracking() const
+{
+	return NumberOfHistoryItemsTracked;
 }
\ No newline at end of file
diff --git a/Engine/Source/Developer/AutomationController/Private/AutomationReport.cpp b/Engine/Source/Developer/AutomationController/Private/AutomationReport.cpp
--- a/Engine/Source/Developer/AutomationController/Private/AutomationReport.cpp
+++ b/Engine/Source/Developer/AutomationController/Private/AutomationReport.cpp
@@ -14,12 +14,19 @@ FAutomationReport::FAutomationReport(FAutomationTestInfo& InTestInfo, bool InIsP
 	, bIsParent( InIsParent )
 	, SupportFlags( 0 )
 	, TestInfo( InTestInfo )
+	, bTrackingHistory(false)
+	, NumRecordsToKeep(0)
 {
 	// Enable smoke tests
 	if ( TestInfo.GetTestType() == EAutomationTestType::ATT_SmokeTest )
 	{
 		bEnabled = true;
 	}
+
+	if (!bIsParent)
+	{
+		LoadHistory();
+	}
 }
 
 
@@ -341,6 +348,234 @@ void FAutomationReport::ResetForExecution(const int32 NumTestPasses)
 }
 
 
+void FAutomationReport::TrackHistory(const bool bShouldTrack, const int32 NumReportsToTrack)
+{
+	bTrackingHistory = bShouldTrack;
+	NumRecordsToKeep = NumReportsToTrack;
+
+	if (bTrackingHistory && ChildReports.Num() == 0)
+	{
+		LoadHistory();
+	}
+
+	//recurse to children
+	for (auto& NextChildReport : ChildReports)
+	{
+		NextChildReport->TrackHistory(bTrackingHistory, NumRecordsToKeep);
+	}
+}
+
+
+void FAutomationReport::AddToHistory()
+{
+	// Dictate the file path we are writing this run as history to.
+	const FDateTime FileDate = FDateTime::Now();
+	const FString FileName = GetDisplayName() + FileDate.ToString() + TEXT(".log");
+	const FString FileLocation = FPaths::ConvertRelativePathToFull(FPaths::AutomationLogDir()) + GetDisplayName();
+	const FString FullPath = FPaths::Combine(*FileLocation, *FileName);
+
+	// Write any Errors and Warnings to the log, if none, then simply report that it was successful.
+	if (FArchive* LogFile = IFileManager::Get().CreateFileWriter(*FullPath))
+	{
+		bool bExportedAnyErrors = false;
+		for (int32 ClusterIndex = 0; ClusterIndex < Results.Num(); ++ClusterIndex)
+		{
+			for (int32 PassIndex = 0; PassIndex < Results[ClusterIndex].Num(); ++PassIndex)
+			{
+				for (int32 ErrorIndex = 0; ErrorIndex < Results[ClusterIndex][PassIndex].Errors.Num(); ++ErrorIndex)
+				{
+					if (!bExportedAnyErrors)
+					{
+						bExportedAnyErrors = true;
+
+						FString ErrorIdentifier(TEXT("<<ERRORS>>"));
+						ErrorIdentifier += LINE_TERMINATOR;
+						
+						LogFile->Serialize(TCHAR_TO_ANSI(*ErrorIdentifier), ErrorIdentifier.Len());
+					}
+					FString NextError = Results[ClusterIndex][PassIndex].Errors[ErrorIndex] + LINE_TERMINATOR;
+					LogFile->Serialize(TCHAR_TO_ANSI(*NextError), NextError.Len());
+				}
+			}
+		}
+
+		bool bExportedAnyWarnings = false;
+		for (int32 ClusterIndex = 0; ClusterIndex < Results.Num(); ++ClusterIndex)
+		{
+			for (int32 PassIndex = 0; PassIndex < Results[ClusterIndex].Num(); ++PassIndex)
+			{
+				for (int32 WarningIndex = 0; WarningIndex < Results[ClusterIndex][PassIndex].Warnings.Num(); ++WarningIndex)
+				{
+					if (!bExportedAnyWarnings)
+					{
+						bExportedAnyWarnings = true;
+
+						FString WarningIdentifier(TEXT("<<WARNINGS>>"));
+						WarningIdentifier += LINE_TERMINATOR;
+
+						LogFile->Serialize(TCHAR_TO_ANSI(*WarningIdentifier), WarningIdentifier.Len());
+					}
+					const FString NextWarning = Results[ClusterIndex][PassIndex].Warnings[WarningIndex] + LINE_TERMINATOR;
+					LogFile->Serialize(TCHAR_TO_ANSI(*NextWarning), NextWarning.Len());
+				}
+			}
+		}
+
+		if (bExportedAnyErrors == false && bExportedAnyWarnings == false)
+		{
+			FString SuccessIdentifier(TEXT("<<SUCCESS>>"));
+			SuccessIdentifier += LINE_TERMINATOR;
+
+			LogFile->Serialize(TCHAR_TO_ANSI(*SuccessIdentifier), SuccessIdentifier.Len());
+		}
+
+		LogFile->Close();
+		delete LogFile;
+
+
+		// Cache an automation history item for tracking in this session
+		TSharedRef<FAutomationHistoryItem> HistoryItem = MakeShareable(new FAutomationHistoryItem);
+		HistoryItem->LogLocation = FileName;
+		HistoryItem->RunDate = FileDate;
+		HistoryItem->RunResult = 
+			(bExportedAnyErrors ? FAutomationHistoryItem::EAutomationHistoryResult::Errors : 
+			(bExportedAnyWarnings ? FAutomationHistoryItem::EAutomationHistoryResult::Warnings : 
+			FAutomationHistoryItem::EAutomationHistoryResult::Successful));
+
+		HistoryItems.Add(HistoryItem);
+	}
+}
+
+
+void FAutomationReport::MaintainHistory()
+{
+	// Find all the logs in this reports log location
+	const FString LogsLocation = FPaths::ConvertRelativePathToFull(FPaths::AutomationLogDir()) + GetDisplayName();
+
+	TArray<FString> LogFiles;
+	IFileManager::Get().FindFiles(LogFiles, *(LogsLocation / "*.log"), true, false);
+
+	// Sort the logs in reverse chronological order
+	struct FLogSortPredicate
+	{
+		FString DisplayName;
+		FLogSortPredicate(const FString& InDisplayName) : DisplayName(InDisplayName) {}
+
+		/** Sort predicate operator */
+		bool operator ()(FString LHS, FString RHS) const
+		{
+			FString LogExt = TEXT(".log");
+
+			FString LHSDateStr = LHS.RightChop(DisplayName.Len());
+			LHSDateStr = LHSDateStr.LeftChop(LogExt.Len());
+
+			FDateTime LHSDate;
+			FDateTime::Parse(LHSDateStr, LHSDate);
+
+			FString RHSDateStr = RHS.RightChop(DisplayName.Len());
+			RHSDateStr = RHSDateStr.LeftChop(LogExt.Len());
+
+			FDateTime RHSDate;
+			FDateTime::Parse(RHSDateStr, RHSDate);
+
+			return LHSDate > RHSDate;
+		}
+	};
+	LogFiles.Sort(FLogSortPredicate(GetDisplayName()));
+
+	// For logs, we keep the number equal to AutomationReportConstants::MaximumLogsToKeep around.
+	// This will mean that we can extend or history to see when changed within the report
+	for (int32 LogIndex = LogFiles.Num() - 1; LogIndex >= AutomationReportConstants::MaximumLogsToKeep; LogIndex--)
+	{
+		check(IFileManager::Get().Delete(*FPaths::Combine(*LogsLocation, *LogFiles[LogIndex])));
+		LogFiles.RemoveAt(LogIndex);
+	}
+
+
+	// Sort the history items in reverse chronological order
+	struct FHistorySortPredicate
+	{
+		FHistorySortPredicate() {}
+
+		/** Sort predicate operator */
+		bool operator ()(const TSharedPtr<FAutomationHistoryItem>& LHS, const TSharedPtr<FAutomationHistoryItem>& RHS) const
+		{
+			check(LHS.IsValid() && RHS.IsValid());
+			return LHS->RunDate > RHS->RunDate;
+		}
+	};
+	HistoryItems.Sort(FHistorySortPredicate());
+
+	for (int32 ItemIndex = HistoryItems.Num() - 1; ItemIndex >= NumRecordsToKeep; ItemIndex--)
+	{
+		HistoryItems.RemoveAt(ItemIndex);
+	}
+}
+
+
+void FAutomationReport::LoadHistory()
+{
+	// Clear out the previous results before we rebuild our list
+	HistoryItems.Empty();
+
+	// Load the logs from this reports automation log location
+	const FString LogsLocation = FPaths::ConvertRelativePathToFull(FPaths::AutomationLogDir()) + GetDisplayName();
+
+	TArray<FString> LogFiles;
+	IFileManager::Get().FindFiles(LogFiles, *(LogsLocation / "*.log"), true, false);
+
+	for (FString& NextLogFile : LogFiles)
+	{
+		FString FileContents;
+		if (FFileHelper::LoadFileToString(FileContents, *FPaths::Combine(*LogsLocation, *NextLogFile)))
+		{
+			TSharedRef<FAutomationHistoryItem> HistoryItem = MakeShareable(new FAutomationHistoryItem);
+
+			// Cache the log location
+			HistoryItem->LogLocation = NextLogFile;
+
+			// Parse the date and time from the log name
+			{
+				FString LogExt = TEXT(".log");
+
+				FString DateStr = NextLogFile.RightChop(GetDisplayName().Len());
+				DateStr = DateStr.LeftChop(LogExt.Len());
+
+				FDateTime::Parse(DateStr, HistoryItem->RunDate);
+			}
+
+			// Parse whether the previous runs had errors, warnings or were successful
+			{
+				if (FileContents.StartsWith(TEXT("<<ERRORS>>")))
+				{
+					HistoryItem->RunResult = FAutomationHistoryItem::EAutomationHistoryResult::Errors;
+				}
+				else if (FileContents.StartsWith(TEXT("<<WARNINGS>>")))
+				{
+					HistoryItem->RunResult = FAutomationHistoryItem::EAutomationHistoryResult::Warnings;
+				}
+				else if (FileContents.StartsWith(TEXT("<<SUCCESS>>")))
+				{
+					HistoryItem->RunResult = FAutomationHistoryItem::EAutomationHistoryResult::Successful;
+				}
+			}
+
+			// Add our log to the tracking
+			HistoryItems.Add(HistoryItem);
+		}
+	}
+
+	// Do a pass on the existing logs for any we no longer wish to maintain.
+	MaintainHistory();
+}
+
+
+const TArray<TSharedPtr<FAutomationHistoryItem>>& FAutomationReport::GetHistory() const
+{
+	return HistoryItems;
+}
+
+
 void FAutomationReport::SetResults( const int32 ClusterIndex, const int32 PassIndex, const FAutomationTestResults& InResults )
 {
 	//verify this is a platform this test is aware of
@@ -358,6 +593,13 @@ void FAutomationReport::SetResults( const int32 ClusterIndex, const int32 PassIn
 	{
 		Results[ClusterIndex][PassIndex].Errors.Add( "No Report Generated" );
 	}
+
+	// If we are tracking history, then export it.
+	if (bTrackingHistory && (InResults.State == EAutomationState::Success || InResults.State == EAutomationState::Fail))
+	{
+		AddToHistory();
+		MaintainHistory();
+	}
 }
 
 
diff --git a/Engine/Source/Developer/AutomationController/Private/AutomationReport.h b/Engine/Source/Developer/AutomationController/Private/AutomationReport.h
--- a/Engine/Source/Developer/AutomationController/Private/AutomationReport.h
+++ b/Engine/Source/Developer/AutomationController/Private/AutomationReport.h
@@ -232,10 +232,38 @@ public:
 	/** Stop the test which is creating this report */
 	virtual void StopRunningTest() OVERRIDE;
 
+	/**
+	* Notification on whether we should, or should not, track this reports history
+	*/
+	virtual void TrackHistory(const bool bShouldTrack, const int32 NumReportsToTrack) override;
+
+	/**
+	* Get the history items of this particular test
+	*
+	* @return A reference to the items of this tests previous runs.
+	*/
+	virtual const TArray<TSharedPtr<FAutomationHistoryItem>>& GetHistory() const override;
+
 	// End IAutomationReport Interface
 	
 private:
 
+	/** 
+	 * Export the current report as part of it's tracked history 
+	 */
+	void AddToHistory();
+
+	/** 
+	 * Load this reports tracked history 
+	 */
+	void LoadHistory();
+
+	/**
+	 * Update what is tracked for this reports history.
+	 */
+	void MaintainHistory();
+
+private:
 
 	/** True if this test should be executed */
 	bool bEnabled;
@@ -269,4 +297,13 @@ private:
 
 	/** Structure holding the test info */
 	FAutomationTestInfo TestInfo;
+
+	/** Flag to determine whether this report should track it's history */
+	bool bTrackingHistory;
+
+	/** Flag to determine how many history items to keep */
+	int32 NumRecordsToKeep;
+
+	/** The collection of history items which holds the results of previous runs. */
+	TArray<TSharedPtr<FAutomationHistoryItem> > HistoryItems;
 };
diff --git a/Engine/Source/Developer/AutomationController/Private/AutomationReportManager.cpp b/Engine/Source/Developer/AutomationController/Private/AutomationReportManager.cpp
--- a/Engine/Source/Developer/AutomationController/Private/AutomationReportManager.cpp
+++ b/Engine/Source/Developer/AutomationController/Private/AutomationReportManager.cpp
@@ -128,6 +128,13 @@ const bool FAutomationReportManager::ExportReport(uint32 FileExportTypeMask, con
 }
 
 
+void FAutomationReportManager::TrackHistory(const bool bShouldTrack, const int32 NumReportsToTrack)
+{
+	// push this request straight to the reports.
+	ReportRoot->TrackHistory(bShouldTrack, NumReportsToTrack);
+}
+
+
 void FAutomationReportManager::AddResultReport(TSharedPtr< IAutomationReport > InReport, const int32 NumDeviceClusters, TArray< FString >& ResultsLog, uint32 FileExportTypeMask)
 {
 	if (InReport->IsEnabled())
diff --git a/Engine/Source/Developer/AutomationController/Private/AutomationReportManager.h b/Engine/Source/Developer/AutomationController/Private/AutomationReportManager.h
--- a/Engine/Source/Developer/AutomationController/Private/AutomationReportManager.h
+++ b/Engine/Source/Developer/AutomationController/Private/AutomationReportManager.h
@@ -75,6 +75,11 @@ public:
 	 */
 	const bool ExportReport( uint32 FileExportTypeMask, const int32 NumDeviceClusters );
 
+	/**
+	 * Notification on whether we should, or should not, track report histories
+	 */
+	void TrackHistory(const bool bShouldTrack, const int32 NumReportsToTrack);
+
 protected:
 
 	/**
diff --git a/Engine/Source/Developer/AutomationController/Public/Interfaces/IAutomationControllerManager.h b/Engine/Source/Developer/AutomationController/Public/Interfaces/IAutomationControllerManager.h
--- a/Engine/Source/Developer/AutomationController/Public/Interfaces/IAutomationControllerManager.h
+++ b/Engine/Source/Developer/AutomationController/Public/Interfaces/IAutomationControllerManager.h
@@ -399,6 +399,25 @@ public:
 	 * Sets the test complete callback delegate.
 	 */
 	virtual void SetTestsCompleteCallback(const FOnAutomationControllerTestsComplete& NewCallback) = 0;
+
+	/**
+	 * Dictate whether to save the reports history and the number of history items to track
+	 *
+	 * @param bShouldTrack - Flag that determines whether to track history
+	 * @param NumReportsToTrack - The number of history reports to keep
+	 */
+	virtual void TrackReportHistory(const bool bShouldTrack, const int32 NumReportsToTrack) = 0;
+
+	/**
+	* Returns whether the controller is tracking history of reports.
+	*/
+	virtual const bool IsTrackingHistory() const = 0;
+
+	/**
+	* Returns the number of history items the controller is maintaining.
+	*/
+	virtual const int32 GetNumberHistoryItemsTracking() const = 0;
+
 public:
 
 	/**
diff --git a/Engine/Source/Developer/AutomationController/Public/Interfaces/IAutomationReport.h b/Engine/Source/Developer/AutomationController/Public/Interfaces/IAutomationReport.h
--- a/Engine/Source/Developer/AutomationController/Public/Interfaces/IAutomationReport.h
+++ b/Engine/Source/Developer/AutomationController/Public/Interfaces/IAutomationReport.h
@@ -98,6 +98,38 @@ struct FAutomationCompleteState
 	uint32 NumDisabledTestsCouldntBeRun;
 };
 
+
+/**
+ * Automation constants we wish to keep across controllers and windows.
+ **/
+namespace AutomationReportConstants
+{
+	const int32 MaximumLogsToKeep = 10;
+}
+
+
+/**
+ * Information for a single automation test's history.
+ **/
+class FAutomationHistoryItem : public TSharedFromThis<FAutomationHistoryItem>
+{
+public:
+	/** The result of this particular test */
+	enum EAutomationHistoryResult
+	{
+		Successful = 0,
+		Warnings,
+		Errors
+	} RunResult;
+
+	/** The file name of the log file. */
+	FString LogLocation;
+
+	/** The date in which this test was completed */
+	FDateTime RunDate;
+};
+
+
 /**
  * Interface for automation test results
  */
@@ -348,4 +380,16 @@ public:
 
 	/** Stop the test which is creating this report */
 	virtual void StopRunningTest() = 0;
+
+	/**
+	 * Notification on whether we should, or should not, track this reports history
+	 */
+	virtual void TrackHistory( const bool bShouldTrack, const int32 NumReportsToTrack ) = 0;
+
+	/**
+	 * Get the history items of this particular test
+	 *
+	 * @return A reference to the items of this tests previous runs.
+	 */
+	virtual const TArray<TSharedPtr<FAutomationHistoryItem>>& GetHistory() const = 0;
 };
diff --git a/Engine/Source/Developer/AutomationWindow/Private/SAutomationTestItem.cpp b/Engine/Source/Developer/AutomationWindow/Private/SAutomationTestItem.cpp
--- a/Engine/Source/Developer/AutomationWindow/Private/SAutomationTestItem.cpp
+++ b/Engine/Source/Developer/AutomationWindow/Private/SAutomationTestItem.cpp
@@ -10,6 +10,120 @@
 #define LOCTEXT_NAMESPACE "AutomationTestItem"
 
 
+/**
+* Implements a Cell widget for the history objects of an automation report.
+*/
+class SAutomationHistoryCell
+	: public SCompoundWidget
+{
+public:
+
+	SLATE_BEGIN_ARGS(SAutomationHistoryCell) { }
+	SLATE_END_ARGS()
+
+public:
+
+	/**
+	* Constructs the widget.
+	*
+	* @param InArgs			- The construction arguments.
+	* @param InHistoryItem	- The history item we are providing the cell for
+	*/
+	void Construct(const FArguments& InArgs, TSharedPtr<IAutomationReport> InHistoryItem);
+
+	/**
+	 * Rebuild the content of the history cell into our Content Area
+	 */
+	void RebuildContentArea();
+
+	// Begin SWidget Interface
+	virtual void Tick(const FGeometry& AllottedGeometry, const double InCurrentTime, const float InDeltaTime) OVERRIDE;
+	// End SWidget Interface
+
+private:
+
+	// A copy of the history items for us to note differences between this and the controller.
+	TArray<TSharedPtr<FAutomationHistoryItem>> HistoryCopy;
+
+	// The automation reports history information
+	TSharedPtr<IAutomationReport> HistoryItem;
+
+	// The widget which holds the content for the History Cell
+	TSharedPtr<SHorizontalBox> ContentArea;
+};
+
+
+BEGIN_SLATE_FUNCTION_BUILD_OPTIMIZATION
+void SAutomationHistoryCell::Construct(const FArguments& InArgs, TSharedPtr<IAutomationReport> InHistoryItem)
+{
+	HistoryItem = InHistoryItem;
+
+	ContentArea = SNew(SHorizontalBox);
+	RebuildContentArea();
+
+	ChildSlot
+	[
+		ContentArea.ToSharedRef()
+	];
+}
+END_SLATE_FUNCTION_BUILD_OPTIMIZATION
+
+
+BEGIN_SLATE_FUNCTION_BUILD_OPTIMIZATION
+void SAutomationHistoryCell::RebuildContentArea()
+{
+	// Clear the previous results before we update the cell again
+	ContentArea->ClearChildren();
+
+	// Create an overview of the previous results in icon form.
+	const TArray<TSharedPtr<FAutomationHistoryItem>> History = HistoryItem->GetHistory();
+	for (const TSharedPtr<FAutomationHistoryItem> HistoryItem : History)
+	{
+		FFormatNamedArguments Args;
+		Args.Add(TEXT("Date"), FText::AsDateTime(HistoryItem->RunDate));
+
+		const FSlateBrush* ResultIcon = nullptr;
+		if (HistoryItem->RunResult == FAutomationHistoryItem::EAutomationHistoryResult::Errors)
+		{
+			Args.Add(TEXT("Result"), LOCTEXT("HasErrors", "had errors"));
+			ResultIcon = FEditorStyle::GetBrush("Automation.Fail");
+		}
+		else if (HistoryItem->RunResult == FAutomationHistoryItem::EAutomationHistoryResult::Warnings)
+		{
+			Args.Add(TEXT("Result"), LOCTEXT("HasWarnings", "had warnings"));
+			ResultIcon = FEditorStyle::GetBrush("Automation.Warning");
+		}
+		else
+		{
+			Args.Add(TEXT("Result"), LOCTEXT("WasSuccessful", "was successful"));
+			ResultIcon = FEditorStyle::GetBrush("Automation.Success");
+		}
+
+		// Add the previous result as an icon representation to the cell
+		ContentArea->AddSlot()
+		[
+			SNew(SImage)
+			.Image(ResultIcon)
+			.ToolTipText(FText::Format(LOCTEXT("ItemTooltip", "{Date} - This run {Result}!"), Args))
+		];
+	}
+}
+END_SLATE_FUNCTION_BUILD_OPTIMIZATION
+
+
+void SAutomationHistoryCell::Tick(const FGeometry& AllottedGeometry, const double InCurrentTime, const float InDeltaTime)
+{
+	const TArray<TSharedPtr<FAutomationHistoryItem>> TestHistory = HistoryItem->GetHistory();
+
+	// if the test history has changed, reflect it in this cell
+	if (TestHistory != HistoryCopy)
+	{
+		RebuildContentArea();
+		HistoryCopy = TestHistory;
+	}
+}
+
+
 /* SAutomationTestItem interface
  *****************************************************************************/
 
@@ -179,6 +293,10 @@ TSharedRef<SWidget> SAutomationTestItem::GenerateWidgetForColumn( const FName& C
 		}
 		return HBox;
 	}
+	else if (ColumnName == AutomationTestWindowConstants::History)
+	{
+		return SNew(SAutomationHistoryCell, TestStatus);
+	}
 	else if( ColumnName == AutomationTestWindowConstants::Timing )
 	{
 		return SNew( STextBlock )
diff --git a/Engine/Source/Developer/AutomationWindow/Private/SAutomationWindow.cpp b/Engine/Source/Developer/AutomationWindow/Private/SAutomationWindow.cpp
--- a/Engine/Source/Developer/AutomationWindow/Private/SAutomationWindow.cpp
+++ b/Engine/Source/Developer/AutomationWindow/Private/SAutomationWindow.cpp
@@ -104,6 +104,10 @@ void SAutomationWindow::Construct( const FArguments& InArgs, const IAutomationCo
 
 	bIsRequestingTests = false;
 
+	// Test history tracking
+	bIsTrackingHistory = AutomationController->IsTrackingHistory();
+	NumHistoryElementsToTrack = AutomationController->GetNumberHistoryItemsTracking();
+	
 	//make the widget for platforms
 	PlatformsHBox = SNew (SHorizontalBox);
 
@@ -124,7 +128,7 @@ void SAutomationWindow::Construct( const FArguments& InArgs, const IAutomationCo
 #endif
 		.HeaderRow
 		(
-		SNew(SHeaderRow)
+		SAssignNew(TestTableHeaderRow,SHeaderRow)
 		+ SHeaderRow::Column( AutomationTestWindowConstants::Title )
 		.FillWidth(300.0f)
 		[
@@ -182,6 +186,19 @@ void SAutomationWindow::Construct( const FArguments& InArgs, const IAutomationCo
 
 		);
 
+	if(bIsTrackingHistory)
+	{
+		TestTableHeaderRow->AddColumn(
+			SHeaderRow::Column(AutomationTestWindowConstants::History)
+			.HAlignHeader(HAlign_Center)
+			.VAlignHeader(VAlign_Center)
+			.HAlignCell(HAlign_Center)
+			.VAlignCell(VAlign_Center)
+			.FixedWidth(100.0f)
+			.DefaultLabel(LOCTEXT("TestHistory", "Test History"))
+		);
+	}
+
 	TSharedRef<SNotificationList> NotificationList = SNew(SNotificationList) .Visibility( EVisibility::HitTestInvisible );
 
 	//build the actual guts of the window
@@ -477,7 +494,7 @@ TSharedRef< SWidget > SAutomationWindow::MakeAutomationWindowToolBar( const TSha
 {
 	struct Local
 	{
-		static void FillToolbar(FToolBarBuilder& ToolbarBuilder, TSharedRef<SWidget> RunTests, TSharedRef<SWidget> Searchbox, TSharedRef<SWidget> PresetBox, TWeakPtr<class SAutomationWindow> InAutomationWindow )
+		static void FillToolbar(FToolBarBuilder& ToolbarBuilder, TSharedRef<SWidget> RunTests, TSharedRef<SWidget> Searchbox, TSharedRef<SWidget> PresetBox, TSharedRef<SWidget> HistoryBox, TWeakPtr<class SAutomationWindow> InAutomationWindow)
 		{
 			ToolbarBuilder.BeginSection("Automation");
 			{
@@ -523,6 +540,20 @@ TSharedRef< SWidget > SAutomationWindow::MakeAutomationWindowToolBar( const TSha
 				ToolbarBuilder.AddWidget( PresetBox );
 			}
 			ToolbarBuilder.EndSection();
+			ToolbarBuilder.BeginSection("History");
+			{
+				ToolbarBuilder.AddWidget(HistoryBox);
+
+				FUIAction DefaultAction;
+				ToolbarBuilder.AddComboButton(
+					DefaultAction,
+					FOnGetContent::CreateStatic(&SAutomationWindow::GenerateTestHistoryMenuContent, InAutomationWindow),
+					LOCTEXT("TestOptions_Label", "Test History Options"),
+					LOCTEXT("TestOptionsToolTip", "Test History Options"),
+					FSlateIcon(FEditorStyle::GetStyleSetName(), "AutomationWindow.TestHistoryOptions"),
+					true);
+			}
+			ToolbarBuilder.EndSection();
 		}
 	};
 
@@ -590,6 +621,44 @@ TSharedRef< SWidget > SAutomationWindow::MakeAutomationWindowToolBar( const TSha
 			.MinDesiredWidth(SearchWidth)
 		];
 
+	const float HistoryWidth = 200.0f;
+	TSharedRef<SWidget> History =
+		SNew(SHorizontalBox)
+		+ SHorizontalBox::Slot()
+		.MaxWidth(HistoryWidth)
+		.AutoWidth()
+		.VAlign(VAlign_Bottom)
+		[
+			SNew(SCheckBox)
+			.Style(FEditorStyle::Get(), "ToggleButtonCheckbox")
+			.Type(ESlateCheckBoxType::ToggleButton)
+			.IsChecked(this, &SAutomationWindow::IsTrackingHistory)
+			.IsEnabled(this, &SAutomationWindow::IsAutomationControllerIdle)
+			.OnCheckStateChanged(this, &SAutomationWindow::OnToggleTrackHistory)
+			[
+				SNew(SVerticalBox)
+				+ SVerticalBox::Slot()
+				.VAlign(VAlign_Center)
+				.HAlign(HAlign_Center)
+				.FillHeight(1.f)
+				.Padding(2.0f)
+				[
+					SNew(SImage)
+					.Image(FEditorStyle::GetBrush("AutomationWindow.TrackHistory"))
+				]
+				+ SVerticalBox::Slot()
+				.Padding(2.0f)
+				.VAlign(VAlign_Bottom)
+				.AutoHeight()
+				[
+					SNew(STextBlock)
+					.ShadowOffset(FVector2D::UnitVector)
+					.Text(LOCTEXT("AutomationHistoryLabel", "Track History"))
+					.IsEnabled(this, &SAutomationWindow::IsAutomationControllerIdle)
+				]
+			]
+		];
+
 	TSharedRef<SWidget> TestPresets = 
 		SNew( SVerticalBox )
 		+SVerticalBox::Slot()
@@ -702,7 +771,7 @@ TSharedRef< SWidget > SAutomationWindow::MakeAutomationWindowToolBar( const TSha
 
 	FToolBarBuilder ToolbarBuilder( InCommandList, FMultiBoxCustomization::None );
 	TWeakPtr<SAutomationWindow> AutomationWindow = SharedThis(this);
-	Local::FillToolbar( ToolbarBuilder, RunTests, Searchbox, TestPresets, AutomationWindow );
+	Local::FillToolbar( ToolbarBuilder, RunTests, Searchbox, TestPresets, History, AutomationWindow );
 
 	// Create the tool bar!
 	return
@@ -958,9 +1027,9 @@ TSharedRef< SWidget > SAutomationWindow::GenerateTestsOptionsMenuContent( )
 			[
 				SNew(SSpinBox<int32>)
 				.MinValue(1)
-				.MaxValue(10)
+				.MaxValue(AutomationReportConstants::MaximumLogsToKeep)
 				.MinSliderValue(1)
-				.MaxSliderValue(10)
+				.MaxSliderValue(AutomationReportConstants::MaximumLogsToKeep)
 				.Value(this,&SAutomationWindow::GetRepeatCount)
 				.OnValueChanged(this,&SAutomationWindow::OnChangeRepeatCount)
 				.IsEnabled( this, &SAutomationWindow::IsAutomationControllerIdle )
@@ -1011,6 +1080,64 @@ TSharedRef< SWidget > SAutomationWindow::GenerateTestsOptionsMenuContent( )
 	return MenuBuilder.MakeWidget();
 }
 
+TSharedRef< SWidget > SAutomationWindow::GenerateTestHistoryMenuContent(TWeakPtr<class SAutomationWindow> InAutomationWindow)
+{
+	TSharedPtr<SAutomationWindow> AutomationWindow(InAutomationWindow.Pin());
+	if (AutomationWindow.IsValid())
+	{
+		return AutomationWindow->GenerateTestHistoryMenuContent();
+	}
+
+	//Return empty menu
+	FMenuBuilder MenuBuilder(true, NULL);
+	MenuBuilder.BeginSection("AutomationWindowTestHistory", LOCTEXT("AutomationWindowTestHistory", "Settings"));
+	MenuBuilder.EndSection();
+	return MenuBuilder.MakeWidget();
+}
+
+TSharedRef< SWidget > SAutomationWindow::GenerateTestHistoryMenuContent()
+{
+	const bool bShouldCloseWindowAfterMenuSelection = true;
+	FMenuBuilder MenuBuilder(bShouldCloseWindowAfterMenuSelection, AutomationWindowActions);
+
+	TSharedRef<SWidget> HistoryCountWidget =
+	SNew(SHorizontalBox)
+	+ SHorizontalBox::Slot()
+	.AutoWidth()
+	.VAlign(VAlign_Bottom)
+	.Padding(FMargin( 0.f, 0.f, 4.f, 0.f ))
+	[
+		SNew(STextBlock)
+		.Text(LOCTEXT("NumberItemsToTrack", "Number Items To Track:"))
+	]
+	+ SHorizontalBox::Slot()
+	.AutoWidth()
+	.VAlign(VAlign_Bottom)
+	[
+		SNew(SBox)
+		.WidthOverride(50.0f)
+		[
+			SNew(SSpinBox<int32>)
+			.MinValue(1)
+			.MaxValue(10)
+			.MinSliderValue(1)
+			.MaxSliderValue(10)
+			.Value(this, &SAutomationWindow::GetTestHistoryCount)
+			.OnValueChanged(this, &SAutomationWindow::OnChangeTestHistoryCount)
+			.IsEnabled(this, &SAutomationWindow::IsAutomationControllerIdle)
+		]
+
+	];
+
+	MenuBuilder.BeginSection("AutomationWindowTestHistory", LOCTEXT("AutomationWindowTestHistory", "Settings"));
+	{
+		MenuBuilder.AddWidget(HistoryCountWidget, FText::GetEmpty());
+	}
+	MenuBuilder.EndSection();
+
+	return MenuBuilder.MakeWidget();
+}
+
 ESlateCheckBoxState::Type SAutomationWindow::IsFullSizeScreenshotsCheckBoxChecked() const
 {
 	return AutomationController->IsUsingFullSizeScreenshots() ? ESlateCheckBoxState::Checked : ESlateCheckBoxState::Unchecked;
@@ -1508,9 +1635,54 @@ void SAutomationWindow::OnToggleErrorFilter()
 	OnRefreshTestCallback();
 }
 
+
+ESlateCheckBoxState::Type SAutomationWindow::IsTrackingHistory() const
+{
+	return bIsTrackingHistory ? ESlateCheckBoxState::Checked : ESlateCheckBoxState::Unchecked;
+}
+
+void SAutomationWindow::OnToggleTrackHistory(ESlateCheckBoxState::Type InState)
+{
+	bIsTrackingHistory = InState == ESlateCheckBoxState::Checked;
+
+	AutomationController->TrackReportHistory(bIsTrackingHistory, NumHistoryElementsToTrack);
+
+	if (bIsTrackingHistory)
+	{
+		TestTableHeaderRow->AddColumn(
+			SHeaderRow::Column(AutomationTestWindowConstants::History)
+			.HAlignHeader(HAlign_Center)
+			.VAlignHeader(VAlign_Center)
+			.HAlignCell(HAlign_Center)
+			.VAlignCell(VAlign_Center)
+			.FixedWidth(100.0f)
+			.DefaultLabel(LOCTEXT("TestHistory", "Test History"))
+		);
+	}
+	else
+	{
+		TestTableHeaderRow->RemoveColumn(AutomationTestWindowConstants::History);
+	}
+
+
+	OnRefreshTestCallback();
+}
+
+void SAutomationWindow::OnChangeTestHistoryCount(int32 InNewValue)
+{
+	NumHistoryElementsToTrack = InNewValue;
+	AutomationController->TrackReportHistory(bIsTrackingHistory, InNewValue);
+}
+
+int32 SAutomationWindow::GetTestHistoryCount() const
+{
+	return NumHistoryElementsToTrack;
+}
+
+
 void SAutomationWindow::OnChangeRepeatCount(int32 InNewValue)
 {
-	AutomationController->SetNumPasses(InNewValue);
+	AutomationController->SetNumPasses(NumHistoryElementsToTrack);
 }
 
 int32 SAutomationWindow::GetRepeatCount() const
diff --git a/Engine/Source/Developer/AutomationWindow/Private/SAutomationWindow.h b/Engine/Source/Developer/AutomationWindow/Private/SAutomationWindow.h
--- a/Engine/Source/Developer/AutomationWindow/Private/SAutomationWindow.h
+++ b/Engine/Source/Developer/AutomationWindow/Private/SAutomationWindow.h
@@ -9,6 +9,7 @@ namespace AutomationTestWindowConstants
 	const FName SmokeTest( TEXT("SmokeTest") );
 	const FName RequiredDeviceCount( TEXT("RequiredDeviceCount") );
 	const FName Status( TEXT("Status") );
+	const FName History( TEXT("History") );
 	const FName Timing( TEXT("Timing") );
 }
 
@@ -166,6 +167,14 @@ private:
 	static TSharedRef< SWidget >GenerateGroupOptionsMenuContent( TWeakPtr<class SAutomationWindow> InAutomationWindow );
 	TSharedRef< SWidget > GenerateGroupOptionsMenuContent( );
 
+	/**
+	* Static: Creates the test history options menu widget
+	*
+	* @return	New widget
+	*/
+	static TSharedRef< SWidget > GenerateTestHistoryMenuContent(TWeakPtr<class SAutomationWindow> InAutomationWindow);
+	TSharedRef< SWidget > GenerateTestHistoryMenuContent();
+
 	/**
 	 * Creates a combo item for the preset list
 	 *
@@ -238,6 +247,10 @@ private:
 	bool IsErrorFilterOn() const;
 	/** Toggles filtering of tests based on error condition */
 	void OnToggleErrorFilter();
+	/** Returns if we're tracking history for automation tests */
+	ESlateCheckBoxState::Type IsTrackingHistory() const;
+	/** Toggles whether we are tracking history of automation tests */
+	void OnToggleTrackHistory(ESlateCheckBoxState::Type InState);
 
 	/** Returns if full size screen shots are enabled */
 	ESlateCheckBoxState::Type IsFullSizeScreenshotsCheckBoxChecked() const;
@@ -262,6 +275,11 @@ private:
 	/** Returns the number of times to repeat the tests */
 	int32 GetRepeatCount() const;
 
+	/** Sets the number of history items we wish to track */
+	void OnChangeTestHistoryCount(int32 InValue);
+	/** Returns the number of history items we currently track */
+	int32 GetTestHistoryCount() const;
+
 	/** Update the test list background style (Editor vs Game) */
 	void UpdateTestListBackgroundStyle();
 
@@ -543,5 +561,12 @@ private:
 	// Holds a pointer to the preset text box
 	TSharedPtr<SEditableTextBox> PresetTextBox;
 
+	// Hold a pointer to the test tables header row.
+	TSharedPtr<SHeaderRow> TestTableHeaderRow;
+
+	// Flag to determine whether we are tracking test history
+	bool bIsTrackingHistory;
 
+	// The number of history elements we are tracking in this session
+	int32 NumHistoryElementsToTrack;
 };
\ No newline at end of file
diff --git a/Engine/Source/Editor/EditorStyle/Private/SlateEditorStyle.cpp b/Engine/Source/Editor/EditorStyle/Private/SlateEditorStyle.cpp
--- a/Engine/Source/Editor/EditorStyle/Private/SlateEditorStyle.cpp
+++ b/Engine/Source/Editor/EditorStyle/Private/SlateEditorStyle.cpp
@@ -5027,6 +5027,7 @@ void FSlateEditorStyle::FStyle::SetupAutomationStyles()
 		Set( "AutomationWindow.SmokeTestFilter.Small", new IMAGE_BRUSH( "Automation/SmokeTestFilter", Icon20x20) );
 		Set( "AutomationWindow.DeveloperDirectoryContent.Small", new IMAGE_BRUSH( "Automation/DeveloperDirectoryContent", Icon20x20) );
 		Set( "AutomationWindow.VisualCommandlet.Small", new IMAGE_BRUSH( "Automation/VisualCommandlet", Icon20x20) );
+		Set( "AutomationWindow.TrackHistory", new IMAGE_BRUSH( "Automation/TrackTestHistory", Icon40x40) );
 
 		//device group settings
 		Set( "AutomationWindow.GroupSettings", new IMAGE_BRUSH( "Automation/Groups", Icon40x40) );
diff --git a/Engine/Source/Runtime/Core/Private/Misc/Paths.cpp b/Engine/Source/Runtime/Core/Private/Misc/Paths.cpp
--- a/Engine/Source/Runtime/Core/Private/Misc/Paths.cpp
+++ b/Engine/Source/Runtime/Core/Private/Misc/Paths.cpp
@@ -165,6 +165,11 @@ FString FPaths::AutomationTransientDir()
 	return FPaths::AutomationDir() + TEXT("Transient/");
 }
 
+FString FPaths::AutomationLogDir()
+{
+	return FPaths::AutomationDir() + TEXT("Logs/");
+}
+
 FString FPaths::CloudDir()
 {
 	return FPaths::GameSavedDir() + TEXT("Cloud/");
diff --git a/Engine/Source/Runtime/Core/Public/Misc/Paths.h b/Engine/Source/Runtime/Core/Public/Misc/Paths.h
--- a/Engine/Source/Runtime/Core/Public/Misc/Paths.h
+++ b/Engine/Source/Runtime/Core/Public/Misc/Paths.h
@@ -197,6 +197,11 @@ public:
 	 */
 	static FString AutomationTransientDir();
 
+	/**
+	* @return The directory for automation log files.
+	*/
+	static FString AutomationLogDir();
+
 	/**
 	 * @return The directory for local files used in cloud emulation or support
 	 */
