Fix for a bug in Start Tests number that creeped in during the history changes. Wrong Spinbox was updated.

#TTP 339298 - EDITOR: Automation: Unable to change the number of test runs

diff --git a/Engine/Source/Developer/AutomationWindow/Private/SAutomationWindow.cpp b/Engine/Source/Developer/AutomationWindow/Private/SAutomationWindow.cpp
--- a/Engine/Source/Developer/AutomationWindow/Private/SAutomationWindow.cpp
+++ b/Engine/Source/Developer/AutomationWindow/Private/SAutomationWindow.cpp
@@ -1027,9 +1027,9 @@ TSharedRef< SWidget > SAutomationWindow::GenerateTestsOptionsMenuContent( )
 			[
 				SNew(SSpinBox<int32>)
 				.MinValue(1)
-				.MaxValue(AutomationReportConstants::MaximumLogsToKeep)
+				.MaxValue(10)
 				.MinSliderValue(1)
-				.MaxSliderValue(AutomationReportConstants::MaximumLogsToKeep)
+				.MaxSliderValue(10)
 				.Value(this,&SAutomationWindow::GetRepeatCount)
 				.OnValueChanged(this,&SAutomationWindow::OnChangeRepeatCount)
 				.IsEnabled( this, &SAutomationWindow::IsAutomationControllerIdle )
@@ -1119,9 +1119,9 @@ TSharedRef< SWidget > SAutomationWindow::GenerateTestHistoryMenuContent()
 		[
 			SNew(SSpinBox<int32>)
 			.MinValue(1)
-			.MaxValue(10)
+			.MaxValue(AutomationReportConstants::MaximumLogsToKeep)
 			.MinSliderValue(1)
-			.MaxSliderValue(10)
+			.MaxSliderValue(AutomationReportConstants::MaximumLogsToKeep)
 			.Value(this, &SAutomationWindow::GetTestHistoryCount)
 			.OnValueChanged(this, &SAutomationWindow::OnChangeTestHistoryCount)
 			.IsEnabled(this, &SAutomationWindow::IsAutomationControllerIdle)
@@ -1682,7 +1682,7 @@ int32 SAutomationWindow::GetTestHistoryCount() const
 
 void SAutomationWindow::OnChangeRepeatCount(int32 InNewValue)
 {
-	AutomationController->SetNumPasses(NumHistoryElementsToTrack);
+	AutomationController->SetNumPasses(InNewValue);
 }
 
 int32 SAutomationWindow::GetRepeatCount() const
