Fixed a handful of LOCTEXT key conflicts in SAutomationWindow.cpp.

diff --git a/Engine/Source/Developer/AutomationWindow/Private/SAutomationWindow.cpp b/Engine/Source/Developer/AutomationWindow/Private/SAutomationWindow.cpp
--- a/Engine/Source/Developer/AutomationWindow/Private/SAutomationWindow.cpp
+++ b/Engine/Source/Developer/AutomationWindow/Private/SAutomationWindow.cpp
@@ -278,7 +278,7 @@ void SAutomationWindow::Construct( const FArguments& InArgs, const IAutomationCo
 							.AutoWidth()
 							[
 								SNew(STextBlock)
-								.Text( LOCTEXT("AutomationTest_GraphicalResults", "Display:"))
+								.Text( LOCTEXT("AutomationTest_Display", "Display:"))
 							]
 
 							+SHorizontalBox::Slot()
@@ -548,8 +548,8 @@ TSharedRef< SWidget > SAutomationWindow::MakeAutomationWindowToolBar( const TSha
 				ToolbarBuilder.AddComboButton(
 					DefaultAction,
 					FOnGetContent::CreateStatic(&SAutomationWindow::GenerateTestHistoryMenuContent, InAutomationWindow),
-					LOCTEXT("TestOptions_Label", "Test History Options"),
-					LOCTEXT("TestOptionsToolTip", "Test History Options"),
+					LOCTEXT("TestHistoryOptions_Label", "Test History Options"),
+					LOCTEXT("TestHistoryOptionsToolTip", "Test History Options"),
 					FSlateIcon(FEditorStyle::GetStyleSetName(), "AutomationWindow.TestHistoryOptions"),
 					true);
 			}
@@ -914,7 +914,7 @@ FString SAutomationWindow::GetPresetComboText() const
 	}
 	else
 	{
-		return LOCTEXT("AutomationPresetLabel", "<Select A Preset>").ToString();
+		return LOCTEXT("AutomationPresetComboLabel", "<Select A Preset>").ToString();
 	}
 }
 
@@ -934,7 +934,7 @@ TSharedRef< SWidget > SAutomationWindow::GenerateGroupOptionsMenuContent( TWeakP
 
 	//Return empty menu
 	FMenuBuilder MenuBuilder( true, NULL );
-	MenuBuilder.BeginSection("AutomationWindowGroupOptions", LOCTEXT("GroupOptions", "Device Group Options"));
+	MenuBuilder.BeginSection("AutomationWindowGroupOptions", LOCTEXT("DeviceGroupOptions", "Device Group Options"));
 	MenuBuilder.EndSection();
 	return MenuBuilder.MakeWidget();
 }
@@ -961,7 +961,7 @@ TSharedRef< SWidget > SAutomationWindow::GenerateGroupOptionsMenuContent( )
 			];
 	}
 
-	MenuBuilder.BeginSection("AutomationWindowGroupDevices", LOCTEXT("GroupOptions", "Group Types"));
+	MenuBuilder.BeginSection("AutomationWindowGroupDevices", LOCTEXT("GroupTypeOptions", "Group Types"));
 	{
 		for( int32 i=0; i<NumFlags;i++ )
 		{
@@ -1042,7 +1042,7 @@ TSharedRef< SWidget > SAutomationWindow::GenerateTestsOptionsMenuContent( )
 		.IsChecked(this, &SAutomationWindow::IsEnableScreenshotsCheckBoxChecked)
 		.OnCheckStateChanged(this, &SAutomationWindow::HandleEnableScreenshotsBoxCheckStateChanged)
 		.Padding(FMargin(4.0f, 0.0f))
-		.ToolTipText(LOCTEXT("AutomationFullSizeScreenshotTip", "If checked, tests that support it will save screenshots"))
+		.ToolTipText(LOCTEXT("AutomationEnableScreenshotsTip", "If checked, tests that support it will save screenshots"))
 		.IsEnabled( this, &SAutomationWindow::IsAutomationControllerIdle )
 		.Content()
 		[
@@ -1055,7 +1055,7 @@ TSharedRef< SWidget > SAutomationWindow::GenerateTestsOptionsMenuContent( )
 		.IsChecked(this, &SAutomationWindow::IsFullSizeScreenshotsCheckBoxChecked)
 		.OnCheckStateChanged(this, &SAutomationWindow::HandleFullSizeScreenshotsBoxCheckStateChanged)
 		.Padding(FMargin(4.0f, 0.0f))
-		.ToolTipText(LOCTEXT("AutomationFullSizeScreenshotTip", "If checked, test that take screenshots will send back full size images instead of thumbnails."))
+		.ToolTipText(LOCTEXT("AutomationFullSizeScreenshotsTip", "If checked, test that take screenshots will send back full size images instead of thumbnails."))
 		.IsEnabled( this, &SAutomationWindow::IsFullSizeScreenshotsOptionEnabled )
 		.Content()
 		[
