Code and documentation cleanup pass

diff --git a/Engine/Source/Developer/AutomationWindow/Private/SAutomationGraphicalResultBox.h b/Engine/Source/Developer/AutomationWindow/Private/SAutomationGraphicalResultBox.h
--- a/Engine/Source/Developer/AutomationWindow/Private/SAutomationGraphicalResultBox.h
+++ b/Engine/Source/Developer/AutomationWindow/Private/SAutomationGraphicalResultBox.h
@@ -1,11 +1,8 @@
 // Copyright 1998-2014 Epic Games, Inc. All Rights Reserved.
 
-/*=============================================================================
-	SAutomationGraphicalResultBox.h: Declares the SAutomationGraphicalResultBox class.
-=============================================================================*/
-
 #pragma once
 
+
 /**
  * The different supported items we can display on the test item
  */
@@ -18,6 +15,7 @@ namespace EAutomationGrapicalDisplayType
 	};
 }
 
+
 /**
  * Implements the automation graphical results box widget.
  */
@@ -35,32 +33,29 @@ public:
 public:
 
 	/**
-	 * Construct this widget
+	 * Construct this widget.
 	 *
-	 * @param InArgs - The declaration data for this widget.
-	 * @param InAutomationController - Reference to the automation controller
+	 * @param InArgs The declaration data for this widget.
+	 * @param InAutomationController Reference to the automation controller.
 	 */
 	void Construct( const FArguments& InArgs, const IAutomationControllerManagerRef& InAutomationController);
 
-	/**
-	 * Clears the current results and any widgets that were created
-	 */
+	/** Clears the current results and any widgets that were created */
 	void ClearResults();
 
-	/**
-	 * Returns if there are any results available to display
-	 */
+	/** Returns if there are any results available to display. */
 	bool HasResults() const;
 
-	/**
-	 * Returns the current display type
-	 */
+	/** Returns the current display type. */
 	EAutomationGrapicalDisplayType::Type GetDisplayType() const;
 
 	/**
 	 * Sets the current display type
+	 *
+	 * @param NewDisplayType The display type to set.
 	 */
 	void SetDisplayType(EAutomationGrapicalDisplayType::Type NewDisplayType);
+
 private:
 
 	/** Holds information on a single test result */
@@ -148,56 +143,50 @@ private:
 	};
 
 	/**
-	 * Returns if there are any results available to display
+	 * Returns if there are any results available to display.
 	 *
-	 * @param InReport - The report to recursively check for enabled tests.
-	 * @param OutReports - The list of enabled test reports
+	 * @param InReport The report to recursively check for enabled tests.
+	 * @param OutReports The list of enabled test reports.
 	 */
 	void GetEnabledReports(TSharedPtr<IAutomationReport> InReport, TArray< TSharedPtr< IAutomationReport > >& OutReports);
 
-	/**
-	 * Populates the ClusterResults with test data pulled from the automation controller
-	 */
+	/** Populates the ClusterResults with test data pulled from the automation controller */
 	void PopulateData();
 
-	/**
-	 * Creates all the slate widgets to display the test results
-	 */
+	/** Creates all the slate widgets to display the test results */
 	void CreateWidgets();
 
-	/**
-	 * This is called when the tests are complete so we can update the widget
-	 */
+	/** This is called when the tests are complete so we can update the widget */
 	void OnFillResults();
 
 	/**
-	 * Gets the color for a test based off the test state and if it has any warnings
+	 * Gets the color for a test based off the test state and if it has any warnings.
 	 *
-	 * @param - TestState - The test state from the automation report (Success, Fail, NotRun, etc.)
-	 * @param - bHasWarnings - If the test had any warning logs.
+	 * @param TestState The test state from the automation report (Success, Fail, NotRun, etc.)
+	 * @param bHasWarnings If the test had any warning logs.
 	 */
 	FSlateColor GetColorForTestState(const EAutomationState::Type TestState, const bool bHasWarnings) const;
 
 	/**
-	 * Gets the text to display for this test item based off the current DisplayType
+	 * Gets the text to display for this test item based off the current DisplayType.
 	 *
-	 * @param - TestName - The name of the test
-	 * @param - TestTime - The duration of the test
+	 * @param TestName The name of the test.
+	 * @param TestTime The duration of the test.
 	 */
 	FText GetTestDisplayText(const FString TestName, const float TestTime) const;
 
-	/** Stores what information should be displayed on the test widgets */
+	/** Stores what information should be displayed on the test widgets. */
 	EAutomationGrapicalDisplayType::Type DisplayType;
 
-	/** The test results grouped by cluster */
+	/** The test results grouped by cluster. */
 	TArray<FClusterResults> ClusterResults;
 
-	/** The time of the longest cluster */
+	/** The time of the longest cluster. */
 	float TotalTestDuration;
 
-	/** Pointer to the root widget that contains all the reslut widgets */
+	/** Pointer to the root widget that contains all the result widgets. */
 	TSharedPtr< SVerticalBox > RootBox;
 
-	/** Pointer to the automation controller so we can get the test results */
+	/** Pointer to the automation controller so we can get the test results. */
 	IAutomationControllerManagerPtr AutomationController;
 };
diff --git a/Engine/Source/Developer/AutomationWindow/Private/SAutomationTestItem.cpp b/Engine/Source/Developer/AutomationWindow/Private/SAutomationTestItem.cpp
--- a/Engine/Source/Developer/AutomationWindow/Private/SAutomationTestItem.cpp
+++ b/Engine/Source/Developer/AutomationWindow/Private/SAutomationTestItem.cpp
@@ -1,9 +1,5 @@
 // Copyright 1998-2014 Epic Games, Inc. All Rights Reserved.
 
-/*=============================================================================
-	SAutomationTestItem.cpp: Implements the SAutomationTestItem class.
-=============================================================================*/
-
 #include "AutomationWindowPrivatePCH.h"
 
 
@@ -26,29 +22,29 @@ public:
 	/**
 	* Constructs the widget.
 	*
-	* @param InArgs			- The construction arguments.
-	* @param InHistoryItem	- The history item we are providing the cell for
+	* @param InArgs The construction arguments.
+	* @param InHistoryItem The history item we are providing the cell for.
 	*/
 	void Construct(const FArguments& InArgs, TSharedPtr<IAutomationReport> InHistoryItem);
 
-	/**
-	 * Rebuild the content of the history cell into our Content Area
-	 */
+	/** Rebuild the content of the history cell into our Content Area. */
 	void RebuildContentArea();
 
-	// Begin SWidget Interface
+public:
+
+	// SWidget Interface
+
 	virtual void Tick(const FGeometry& AllottedGeometry, const double InCurrentTime, const float InDeltaTime) override;
-	// End SWidget Interface
 
 private:
 
-	// A copy of the history items for us to note differences between this and the controller.
+	/** A copy of the history items for us to note differences between this and the controller. */
 	TArray<TSharedPtr<FAutomationHistoryItem>> HistoryCopy;
 
-	// The automation reports history information
+	/** The automation reports history information. */
 	TSharedPtr<IAutomationReport> HistoryItem;
 
-	// The widget which holds the content for the History Cell
+	/** The widget which holds the content for the History Cell. */
 	TSharedPtr<SHorizontalBox> ContentArea;
 };
 
@@ -313,7 +309,7 @@ END_SLATE_FUNCTION_BUILD_OPTIMIZATION
  *****************************************************************************/
 const FSlateBrush* SAutomationTestItem::GetSmokeTestImage() const
 {
-	const FSlateBrush* ImageToUse = NULL;
+	const FSlateBrush* ImageToUse = nullptr;
 	if ( TestStatus->GetTestType() & EAutomationTestType::ATT_SmokeTest )
 	{
 		if ( TestStatus->IsParent() )
diff --git a/Engine/Source/Developer/AutomationWindow/Private/SAutomationTestItem.h b/Engine/Source/Developer/AutomationWindow/Private/SAutomationTestItem.h
--- a/Engine/Source/Developer/AutomationWindow/Private/SAutomationTestItem.h
+++ b/Engine/Source/Developer/AutomationWindow/Private/SAutomationTestItem.h
@@ -1,11 +1,8 @@
 // Copyright 1998-2014 Epic Games, Inc. All Rights Reserved.
 
-/*=============================================================================
-	SAutomationTestItem.h: Declares the SAutomationTestItem class.
-=============================================================================*/
-
 #pragma once
 
+
 DECLARE_DELEGATE_OneParam(FOnItemCheckedStateChanged, TSharedPtr<IAutomationReport> /*TestStatus*/);
 
 
@@ -37,7 +34,7 @@ public:
 	/**
 	 * Generates the widget for the specified column.
 	 *
-	 * @param ColumnName - The name of the column to generate the widget for.
+	 * @param ColumnName The name of the column to generate the widget for.
 	 * @return The widget.
 	 */
 	virtual TSharedRef<SWidget> GenerateWidgetForColumn( const FName& ColumnName ) override;
@@ -61,8 +58,8 @@ protected:
 	/**
 	 * Returns the tool tip for the automation test result.
 	 *
-	 * @param ClusterIndex - the cluster index.
-	 * @return The tooltip text
+	 * @param ClusterIndex The cluster index.
+	 * @return The tooltip text.
 	 */
 	FText GetTestToolTip( int32 ClusterIndex ) const;
 
@@ -76,7 +73,7 @@ protected:
 	/**
 	 * Returns color that indicates test status per cluster.
 	 *
-	 * @param ClusterIndex - the cluster index.
+	 * @param ClusterIndex The cluster index.
 	 * @return The background Color.
 	 */
 	FSlateColor ItemStatus_BackgroundColor( const int32 ClusterIndex ) const;
@@ -91,8 +88,8 @@ protected:
 	/**
 	 * Helper to ensure throbber is visible when in process and icon is visible otherwise.
 	 *
-	 * @param ClusterIndex - the cluster index.
- 	 * @param bForInProcessThrobber - If throbbing is in process.
+	 * @param ClusterIndex The cluster index.
+ 	 * @param bForInProcessThrobber If throbbing is in process.
 	 * @return The visibility state.
 	 */
 	EVisibility ItemStatus_GetStatusVisibility( const int32 ClusterIndex, const bool bForInProcessThrobber ) const;
@@ -100,7 +97,7 @@ protected:
 	/**
 	 * Color of progress bar for internal tree test nodes.
 	 *
-	 * @param ClusterIndex - the cluster index.
+	 * @param ClusterIndex The cluster index.
 	 * @return The progress bar color.
 	 */
 	FSlateColor ItemStatus_ProgressColor( const int32 ClusterIndex ) const;
@@ -108,7 +105,7 @@ protected:
 	/**
 	 * Returns percent completion for an internal tree node for all enabled child tests.
 	 *
-	 * @param ClusterIndex - the cluster index.
+	 * @param ClusterIndex The cluster index.
 	 * @return The percent complete.
 	 */
 	TOptional<float> ItemStatus_ProgressFraction( const int32 ClusterIndex ) const;
@@ -116,38 +113,34 @@ protected:
 	/**
 	 * Returns image that denotes the status of a particular test on the given platform cluster.
 	 *
-	 * @param ClusterIndex - the cluster index.
+	 * @param ClusterIndex The cluster index.
 	 * @return the status image.
 	 */
 	const FSlateBrush* ItemStatus_StatusImage( const int32 ClusterIndex ) const;
 	
 	/**
-	 * The number of participants required for this test item in string form
+	 * The number of participants required for this test item in string form.
 	 *
-	 * @return The number of participants needed
+	 * @return The number of participants needed.
 	 */
 	FText ItemStatus_NumParticipantsRequiredText( ) const;
 
 private:
 
-	/**
-	 * Handle the testing checkbox click.
-	 *
-	 * @param The checkbox state.
-	 */
+	/** Handle the testing checkbox click. */
 	void HandleTestingCheckbox_Click( ESlateCheckBoxState::Type );
 
 private:
 
-	// The column width
+	/** The column width. */
 	float ColumnWidth;
 
-	// Holds the highlight string for the automation test.
+	/** Holds the highlight string for the automation test. */
 	TAttribute<FText> HighlightText;
 
-	// Holds the automation report
+	/** Holds the automation report. */
 	TSharedPtr<IAutomationReport> TestStatus;
 
-	// Holds a delegate to be invoked when the check box state changed.
+	/** Holds a delegate to be invoked when the check box state changed. */
 	FOnItemCheckedStateChanged OnCheckedStateChangedDelegate;
 };
diff --git a/Engine/Source/Developer/AutomationWindow/Private/SAutomationTestItemContextMenu.h b/Engine/Source/Developer/AutomationWindow/Private/SAutomationTestItemContextMenu.h
--- a/Engine/Source/Developer/AutomationWindow/Private/SAutomationTestItemContextMenu.h
+++ b/Engine/Source/Developer/AutomationWindow/Private/SAutomationTestItemContextMenu.h
@@ -1,11 +1,8 @@
 // Copyright 1998-2014 Epic Games, Inc. All Rights Reserved.
 
-/*=============================================================================
-	SAutomationTestItemContextMenu.h: Declares the SAutomationTestItemContextMenu class.
-=============================================================================*/
-
 #pragma once
 
+
 #if WITH_EDITOR
 
 #define LOCTEXT_NAMESPACE "SAutomationTestItemContextMenu"
@@ -23,8 +20,8 @@ public:
 	/**
 	 * Construct this widget.
 	 *
-	 * @param InArgs - The declaration data for this widget.
-	 * @param InSessionManager - The session to use.
+	 * @param InArgs The declaration data for this widget.
+	 * @param InSessionManager The session to use.
 	 */
 	void Construct( const FArguments& InArgs, const FString& InAssetName )
 	{
@@ -63,9 +60,7 @@ protected:
 
 private:
 
-	/**
-	 * Handle the context menu closing down. If an asset is selected, request that it gets loaded
-	 */
+	/** Handle the context menu closing down. If an asset is selected, request that it gets loaded */
 	void HandleContextItemTerminate( )
 	{
 		if (Assetname.Len() > 0)
@@ -76,11 +71,11 @@ private:
 
 private:
 
-	// Holds the selected asset name
+	/** Holds the selected asset name. */
 	FString Assetname;
 };
 
 
 #undef LOCTEXT_NAMESPACE
 
-#endif
\ No newline at end of file
+#endif //WITH_EDITOR
\ No newline at end of file
diff --git a/Engine/Source/Developer/AutomationWindow/Private/SAutomationTestTreeView.h b/Engine/Source/Developer/AutomationWindow/Private/SAutomationTestTreeView.h
--- a/Engine/Source/Developer/AutomationWindow/Private/SAutomationTestTreeView.h
+++ b/Engine/Source/Developer/AutomationWindow/Private/SAutomationTestTreeView.h
@@ -1,11 +1,8 @@
 // Copyright 1998-2014 Epic Games, Inc. All Rights Reserved.
 
-/*=============================================================================
-	SAutomationTestTreeView.h: Declares the SAutomationTestTreeView class.
-=============================================================================*/
-
 #pragma once
 
+
 /**
  * Implements the automation test tree view.  Same slate interface as STreeView but adds the ability to clear the internal widget cache.
  */
diff --git a/Engine/Source/Developer/AutomationWindow/Private/SAutomationWindow.cpp b/Engine/Source/Developer/AutomationWindow/Private/SAutomationWindow.cpp
--- a/Engine/Source/Developer/AutomationWindow/Private/SAutomationWindow.cpp
+++ b/Engine/Source/Developer/AutomationWindow/Private/SAutomationWindow.cpp
@@ -6,8 +6,10 @@
 #include "AutomationPresetManager.h"
 #include "SAutomationWindow.h"
 
+
 #define LOCTEXT_NAMESPACE "AutomationTest"
 
+
 //////////////////////////////////////////////////////////////////////////
 // FAutomationWindowCommands
 
@@ -838,7 +840,7 @@ void SAutomationWindow::HandlePresetTextCommited( const FText& CommittedText, ET
 		if( bAddingTestPreset )
 		{
 			bAddingTestPreset = false;
-			SelectedPreset = NULL;
+			SelectedPreset = nullptr;
 			PresetComboBox->ClearSelection();
 			PresetTextBox->SetText(FText());
 		}
@@ -900,7 +902,7 @@ FReply SAutomationWindow::HandleRemovePresetClicked()
 	if(SelectedPreset.IsValid())
 	{
 		TestPresetManager->RemovePreset(SelectedPreset.ToSharedRef());
-		SelectedPreset = NULL;
+		SelectedPreset = nullptr;
 		PresetComboBox->ClearSelection();
 	}
 	return FReply::Handled();
@@ -933,7 +935,7 @@ TSharedRef< SWidget > SAutomationWindow::GenerateGroupOptionsMenuContent( TWeakP
 	}
 
 	//Return empty menu
-	FMenuBuilder MenuBuilder( true, NULL );
+	FMenuBuilder MenuBuilder( true, nullptr );
 	MenuBuilder.BeginSection("AutomationWindowGroupOptions", LOCTEXT("DeviceGroupOptions", "Device Group Options"));
 	MenuBuilder.EndSection();
 	return MenuBuilder.MakeWidget();
@@ -1000,7 +1002,7 @@ TSharedRef< SWidget > SAutomationWindow::GenerateTestsOptionsMenuContent( TWeakP
 	}
 
 	//Return empty menu
-	FMenuBuilder MenuBuilder( true, NULL );
+	FMenuBuilder MenuBuilder( true, nullptr );
 	MenuBuilder.BeginSection("AutomationWindowRunTest", LOCTEXT("RunTestOptions", "Advanced Settings"));
 	MenuBuilder.EndSection();
 	return MenuBuilder.MakeWidget();
@@ -1089,7 +1091,7 @@ TSharedRef< SWidget > SAutomationWindow::GenerateTestHistoryMenuContent(TWeakPtr
 	}
 
 	//Return empty menu
-	FMenuBuilder MenuBuilder(true, NULL);
+	FMenuBuilder MenuBuilder(true, nullptr);
 	MenuBuilder.BeginSection("AutomationWindowTestHistory", LOCTEXT("AutomationWindowTestHistory", "Settings"));
 	MenuBuilder.EndSection();
 	return MenuBuilder.MakeWidget();
@@ -1177,7 +1179,7 @@ TSharedPtr<SWidget> SAutomationWindow::HandleAutomationListContextMenuOpening()
 		}
 	}		
 
-	return NULL;
+	return nullptr;
 }
 #endif
 
diff --git a/Engine/Source/Developer/AutomationWindow/Private/SAutomationWindow.h b/Engine/Source/Developer/AutomationWindow/Private/SAutomationWindow.h
--- a/Engine/Source/Developer/AutomationWindow/Private/SAutomationWindow.h
+++ b/Engine/Source/Developer/AutomationWindow/Private/SAutomationWindow.h
@@ -2,6 +2,7 @@
 
 #pragma once
 
+
 /** Columns for the test tree view */
 namespace AutomationTestWindowConstants
 {
@@ -13,6 +14,7 @@ namespace AutomationTestWindowConstants
 	const FName Timing( TEXT("Timing") );
 }
 
+
 /** The type of background style to use for the test list widget */
 namespace EAutomationTestBackgroundStyle
 {
@@ -24,6 +26,7 @@ namespace EAutomationTestBackgroundStyle
 	};
 }
 
+
 /**
  * Implements the main UI Window for hosting all automation tests.
  */
@@ -59,14 +62,10 @@ public:
 
 public:
 
-	/**
-	 * Default constructor.
-	 */
+	/** Default constructor. */
 	SAutomationWindow();
 
-	/**
-	 * Default destructor.
-	 */
+	/** Destructor. */
 	virtual ~SAutomationWindow();
 
 public:
@@ -77,18 +76,18 @@ public:
 	void Construct(const FArguments& InArgs, const IAutomationControllerManagerRef& AutomationController, const ISessionManagerRef& InSessionManager);
 
 	/**
-	* Check tests aren't running
+	* Check tests aren't running.
 	*
-	* @return true if the tests aren't running
+	* @return true if the tests aren't running.
 	*/
 	bool IsAutomationControllerIdle() const;
 
 protected:
 
 	/**
-	* Checks the list of selected rows to see if multiple rows are selected
+	* Checks the list of selected rows to see if multiple rows are selected.
 	*
-	* @return true if multiple rows are selected
+	* @return true if multiple rows are selected.
 	*/
 	bool AreMultipleRowsSelected() {return TestTable->GetSelectedItems().Num()>1;}
 
@@ -100,115 +99,122 @@ protected:
 	void ChangeTheSelectionToThisRow(TSharedPtr< IAutomationReport > ThisRow);
 
 	/**
-	* Tests if the given row is in the list of selected rows
+	* Tests if the given row is in the list of selected rows.
 	*
-	* @param ThisRow a Report in the list of all report
-	* @return true if the report passed in is in the list of selected rows
+	* @param ThisRow a Report in the list of all report.
+	* @return true if the report passed in is in the list of selected rows.
 	*/
 	bool IsRowSelected(TSharedPtr< IAutomationReport >  ThisRow);
 
 	/**
 	* Sets the enabled value of the selected rows to given value.
 	*
-	* @param InChecked Set all the selected rows to the enabled state to this value
+	* @param InChecked Set all the selected rows to the enabled state to this value.
 	*/
 	void SetAllSelectedTestsChecked( bool InChecked );
 
 	/**
-	* Checks the list of selected rows to see if any are enabled
+	* Checks the list of selected rows to see if any are enabled.
 	*
-	* @return true if any of the selected rows are enabled
+	* @return true if any of the selected rows are enabled.
 	*/
 	bool IsAnySelectedRowEnabled();
 
-	/** Overridden from SWidget: Called when a key is pressed down - capturing copy */
+	/** Overridden from SWidget: Called when a key is pressed down - capturing copy. */
 	virtual FReply OnKeyDown( const FGeometry& MyGeometry, const FKeyboardEvent& InKeyboardEvent ) override;
 
-	/** Overridden from SWidget: Called after a key is released */
+	/** Overridden from SWidget: Called after a key is released. */
 	virtual FReply OnKeyUp( const FGeometry& InGeometry, const FKeyboardEvent& InKeyboardEvent ) override;
 
 private:
 
 	/**
-	* Handle for a test item's checkbox being clicked
+	* Handle for a test item's check box being clicked.
 	*
-	* @param TestStatus The test row details
+	* @param TestStatus The test row details.
 	*/
 	void HandleItemCheckBoxCheckedStateChanged( TSharedPtr< IAutomationReport > TestStatus );
 
-	// Callback for getting the enabled state of a test item.
+	/** Callback for getting the enabled state of a test item. */
 	bool HandleItemCheckBoxIsEnabled( ) const;
 
 	/** Create the UI commands for the toolbar */
 	void CreateCommands();
 
 	/**
-	 * Static: Creates a toolbar widget for the main automation window
+	 * Static: Creates a toolbar widget for the main automation window.
 	 *
-	 * @return	New widget
+	 * @return The new widget.
 	 */
 	static TSharedRef< SWidget > MakeAutomationWindowToolBar( const TSharedRef<FUICommandList>& InCommandList, TSharedPtr<class SAutomationWindow> InLevelEditor );
 	TSharedRef< SWidget > MakeAutomationWindowToolBar( const TSharedRef<FUICommandList>& InCommandList );
 
 	/**
-	 * Static: Creates the test options menu widget
+	 * Static: Creates the test options menu widget.
 	 *
-	 * @return	New widget
+	 * @return	The new widget.
 	 */
 	static TSharedRef< SWidget >GenerateTestsOptionsMenuContent( TWeakPtr<class SAutomationWindow> InAutomationWindow );
 	TSharedRef< SWidget > GenerateTestsOptionsMenuContent( );
 
-	
 	/**
-	 * Static: Creates the group flag options menu widget
+	 * Static: Creates the group flag options menu widget.
 	 *
-	 * @return	New widget
+	 * @return The new widget.
 	 */
 	static TSharedRef< SWidget >GenerateGroupOptionsMenuContent( TWeakPtr<class SAutomationWindow> InAutomationWindow );
 	TSharedRef< SWidget > GenerateGroupOptionsMenuContent( );
 
 	/**
-	* Static: Creates the test history options menu widget
+	* Static: Creates the test history options menu widget.
 	*
-	* @return	New widget
+	* @return The new widget.
 	*/
 	static TSharedRef< SWidget > GenerateTestHistoryMenuContent(TWeakPtr<class SAutomationWindow> InAutomationWindow);
 	TSharedRef< SWidget > GenerateTestHistoryMenuContent();
 
 	/**
-	 * Creates a combo item for the preset list
+	 * Creates a combo item for the preset list.
 	 *
-	 * @return	New combo item widget
+	 * @return New combo item widget.
 	 */
 	TSharedRef<SWidget> GeneratePresetComboItem(TSharedPtr<FAutomationTestPreset> InItem);
 
 	/** 
-	 * Populates OutSearchStrings with the strings that should be used in searching
-	 * @param Report - the automation report to get a text description from
-	 * @param OutSearchStrings - an array of stings to use
+	 * Populates OutSearchStrings with the strings that should be used in searching.
+	 *
+	 * @param Report The automation report to get a text description from.
+	 * @param OutSearchStrings An array of stings to use.
 	 */
 	void PopulateReportSearchStrings( const TSharedPtr< IAutomationReport >& Report, OUT TArray< FString >& OutSearchStrings ) const;
 
 	/** Gets children tests for a node in the hierarchy */
 	void OnGetChildren(TSharedPtr<IAutomationReport> InItem, TArray<TSharedPtr<IAutomationReport> >& OutItems);
+	
 	/** Callback for a new test being selected */
 	void OnTestSelectionChanged(TSharedPtr<IAutomationReport> Selection, ESelectInfo::Type SelectInfo);
+	
 	/** Called when the header checkbox's state is changed. Will go through all the Unit Tests available and change the state of their checkbox. */
 	void HeaderCheckboxStateChange(ESlateCheckBoxState::Type InCheckboxState);	/** Creates table row for each visible automation test */
+	
 	/** Rebuilds the platform icon header */
 	void RebuildPlatformIcons();
+	
 	/** Generate the device tooltip */
 	FText CreateDeviceTooltip(int32 ClusterIndex);
+	
 	/** Clear the UI and icon header */
 	void ClearAutomationUI ();
 
 	/** Generates a row widget for a automation test */
 	TSharedRef<ITableRow> OnGenerateWidgetForTest(TSharedPtr<IAutomationReport> InItem, const TSharedRef<STableViewBase>& OwnerTable );
+	
 	/** Generates a row widget for the log list view. */
 	TSharedRef<ITableRow> OnGenerateWidgetForLog(TSharedPtr<FAutomationOutputMessage> Message, const TSharedRef<STableViewBase>& OwnerTable);
 
 	/** Returns number of enabled tests (regardless of visibility) */
 	FString OnGetNumEnabledTestsString() const;
+	
 	/** Returns number of workers in a device cluster */
 	FString OnGetNumDevicesInClusterString(const int32 ClusterIndex) const;
 
@@ -221,6 +227,7 @@ private:
 	//TODO AUTOMATION - remove if this step is automatic
 	/** Updates list of all the tests */
 	void ListTests();
+	
 	/** Goes through all selected tests and runs them. */
 	FReply RunTests();
 
@@ -229,26 +236,37 @@ private:
 
 	/** Returns true if automation test includes visual commandlet*/
 	bool IsVisualCommandletFilterOn() const;
+	
 	/** Toggles filter of visual commandlet */
 	void OnToggleVisualCommandletFilter();
+	
 	/** Returns if we're considering tests on content within the developer folders */
 	bool IsDeveloperDirectoryIncluded() const;
+	
 	/** Toggles the consideration of tests within developer folders */
 	void OnToggleDeveloperDirectoryIncluded();
+	
 	/** Returns if we're filtering based on if the test is a "smoke" test */
 	bool IsSmokeTestFilterOn() const;
+	
 	/** Toggles filtering of tests based on smoke test status */
 	void OnToggleSmokeTestFilter();
+	
 	/** Returns if we're filtering based on if the test returned any warnings */
 	bool IsWarningFilterOn() const;
+	
 	/** Toggles filtering of tests based on warning condition */
 	void OnToggleWarningFilter();
+	
 	/** Returns if we're filtering based on if the test returned any errors */
 	bool IsErrorFilterOn() const;
+	
 	/** Toggles filtering of tests based on error condition */
 	void OnToggleErrorFilter();
+	
 	/** Returns if we're tracking history for automation tests */
 	ESlateCheckBoxState::Type IsTrackingHistory() const;
+	
 	/** Toggles whether we are tracking history of automation tests */
 	void OnToggleTrackHistory(ESlateCheckBoxState::Type InState);
 
@@ -267,16 +285,19 @@ private:
 
 	/** Returns if a device group is enabled */
 	ESlateCheckBoxState::Type IsDeviceGroupCheckBoxIsChecked(const int32 DeviceGroupFlag) const;
+	
 	/** Toggles a device group flag */
 	void HandleDeviceGroupCheckStateChanged(ESlateCheckBoxState::Type CheckBoxState, const int32 DeviceGroupFlag);
 	
 	/** Sets the number of times to repeat the tests */
 	void OnChangeRepeatCount(int32 InNewValue);
+	
 	/** Returns the number of times to repeat the tests */
 	int32 GetRepeatCount() const;
 
 	/** Sets the number of history items we wish to track */
 	void OnChangeTestHistoryCount(int32 InValue);
+	
 	/** Returns the number of history items we currently track */
 	int32 GetTestHistoryCount() const;
 
@@ -284,196 +305,170 @@ private:
 	void UpdateTestListBackgroundStyle();
 
 	/**
-	* Gets the extention for the small brush, if enabled
+	* Gets the extention for the small brush, if enabled.
 	*
-	* @return - The string to append to the brush path
+	* @return The string to append to the brush path.
 	*/
 	FString GetSmallIconExtension() const;
 
 	/**
-	* Gets whether we should show large tool bar button details
+	* Gets whether we should show large tool bar button details.
 	*
-	* @return - The visibility setting for large tool bar details
+	* @return The visibility setting for large tool bar details.
 	*/
 	EVisibility GetLargeToolBarVisibility() const;
 
 	/**
-	* Gets a brush for the automation start / stop state
+	* Gets a brush for the automation start / stop state.
 	*
-	* @return - A start or stop icon - depending on state
+	* @return - A start or stop icon - depending on state.
 	*/
 	const FSlateBrush* GetRunAutomationIcon() const;
 
 	/**
-	* Gets a label for the automation start / stop state
+	* Gets a label for the automation start / stop state.
 	*
-	* @return A start or stop tests label - depending on state
+	* @return A start or stop tests label Depending on state.
 	*/
 	FString GetRunAutomationLabel() const;
 
 	/**
-	* Gets the brush to use for the test list background
+	* Gets the brush to use for the test list background.
 	*
-	* @return The brush to use depending on the value of TestBackgroundType
+	* @return The brush to use depending on the value of TestBackgroundType.
 	*/
 	const FSlateBrush* GetTestBackgroundBorderImage() const;
 
 	/**
-	* Recursively expand the tree nodes
+	* Recursively expand the tree nodes.
 	*
-	* @param InReport The current report
-	* @param ShouldExpand - Should we expand the item
+	* @param InReport The current report.
+	* @param ShouldExpand Should we expand the item.
 	*/
 	void ExpandTreeView( TSharedPtr< IAutomationReport > InReport, const bool ShouldExpand );
 
-	/**
-	* Update the highlight string in the automation reports
-	*/
+	/** Update the highlight string in the automation reports. */
 	FText HandleAutomationHighlightText( ) const;
 
-	// Callback for determining the visibility of the 'Select a session' overlay.
+	/** Callback for determining the visibility of the 'Select a session' overlay. */
 	EVisibility HandleSelectSessionOverlayVisibility( ) const;
 
-	// Callback for determining whether a session can be selected in the session manager.
+	/** Callback for determining whether a session can be selected in the session manager. */
 	void HandleSessionManagerCanSelectSession( const ISessionInfoPtr& Session, bool& CanSelect );
 
 	/**
-	* Session selection has changed in the session manager
-	*/
+	 * Session selection has changed in the session manager
+	 *
+	 * @param SelectedSession The session that was selected.
+	 */
 	void HandleSessionManagerSelectionChanged( const ISessionInfoPtr& SelectedSession );
 
-	/** 
-	* Called when the session manager updates an instances
-	*/
+	/** Called when the session manager updates an instances. */
 	void HandleSessionManagerInstanceChanged();
 
 	/**
-	* Should the automation run button be enabled
-	*
-	* @return - true if it should be enabled
-	*/
+	 * Should the automation run button be enabled.
+	 *
+	 * @return true if it should be enabled.
+	 */
 	bool IsAutomationRunButtonEnabled() const;
 
 	/**
-	* Set whether tests are available to run
-	*
-	* @param - The Automation controller state
-	*/
+ 	 * Set whether tests are available to run.
+	 *
+	 * @param The Automation controller state.
+	 */
 	void OnTestAvailableCallback( EAutomationControllerModuleState::Type InAutomationControllerState );
 
-	/**
-	 * Copies the selected log messages to the clipboard.
-	 */
+	/** Copies the selected log messages to the clipboard. */
 	void CopyLog();
 
 #if WITH_EDITOR
 	/**
-	* Handle the context menu opening for automation reports
-	*
-	* @return the context window
-	*/
+	 * Handle the context menu opening for automation reports.
+	 *
+	 * @return the context window.
+	 */
 	TSharedPtr<SWidget> HandleAutomationListContextMenuOpening( );
 #endif
 
-	/**
-	 * Handles the new preset button being clicked
-	 */
+	/** Handles the new preset button being clicked. */
 	FReply HandleNewPresetClicked();
 
-	/**
-	 * Handles the save preset button being clicked
-	 */
+	/** Handles the save preset button being clicked. */
 	FReply HandleSavePresetClicked();
 
-	/**
-	 * Handles the remove preset button being clicked
-	 */
+	/** Handles the remove preset button being clicked. */
 	FReply HandleRemovePresetClicked();
 
-	/**
-	* Should the add preset button be enabled
-	*/
+	/** Should the add preset button be enabled. */
 	bool IsAddButtonEnabled() const;
 
-	/**
-	* Should the save preset button be enabled
-	*/
+	/** Should the save preset button be enabled. */
 	bool IsSaveButtonEnabled() const;
 
-	/**
-	* Should the remove preset button be enabled
-	*/
+	/** Should the remove preset button be enabled. */
 	bool IsRemoveButtonEnabled() const;
 
-	/**
-	* Handles if the preset combo box should be visible
-	*/
+	/** Handles if the preset combo box should be visible. */
 	EVisibility HandlePresetComboVisibility( ) const;
 
-	/**
-	* Handles if the add preset text box should be visible
-	*/
+	/** Handles if the add preset text box should be visible. */
 	EVisibility HandlePresetTextVisibility( ) const;
 
-	/**
-	* Called when the user commits the text in the add preset text box
-	*/
+	/** Called when the user commits the text in the add preset text box. */
 	void HandlePresetTextCommited( const FText& CommittedText, ETextCommit::Type CommitType );
 
-	/**
-	* Called when the user selects a new preset from the preset combo box
-	*/
+	/** Called when the user selects a new preset from the preset combo box. */
 	void HandlePresetChanged( TSharedPtr<FAutomationTestPreset> Item, ESelectInfo::Type SelectInfo );
 
-	/**
-	* Expands the test tree to show all enabled tests
-	*/
+	/** Expands the test tree to show all enabled tests. */
 	void ExpandEnabledTests( TSharedPtr< IAutomationReport > InReport );
 
-	/**
-	* Gets the text to display for the preset combo box
-	*/
+	/** Gets the text to display for the preset combo box. */
 	FString GetPresetComboText() const;
 
 	/**
-	* Handle the copy button clicked in the command bar
-	*
-	* @return If the command was handled or not
-	*/
+	 * Handle the copy button clicked in the command bar.
+	 *
+	 * @return If the command was handled or not.
+	 */
 	FReply HandleCommandBarCopyLogClicked();
 
 	/**
-	* Handle the log selection changed
-	*
-	* @param The selected item
-	* @param SelectInfo Provides context on how the selection changed
-	*/
+	 * Handle the log selection changed.
+	 *
+	 * @param The selected item.
+	 * @param SelectInfo Provides context on how the selection changed.
+	 */
 	void HandleLogListSelectionChanged( TSharedPtr<FAutomationOutputMessage> InItem, ESelectInfo::Type SelectInfo );
 
 	/**
-	* Gets the visibility of the test log window
-	*/
+	 * Gets the visibility of the test log window.
+	 *
+	 * @return The window's visibility.
+	 * @see GetTestGraphVisibility
+	 */
 	EVisibility GetTestLogVisibility( ) const;
 
 	/**
-	* Gets the visibility of the graphical test result window
-	*/
+	 * Gets the visibility of the graphical test result window.
+	 *
+	 * @return The window's visibility.
+	 * @see GetTestLogVisibility
+	 */
 	EVisibility GetTestGraphVisibility( ) const;
 
-	/**
-	* Handles changing the display type for the graphical results window
-	*/
+	/** Handles changing the display type for the graphical results window. */
 	void HandleResultDisplayTypeStateChanged( ESlateCheckBoxState::Type NewRadioState, EAutomationGrapicalDisplayType::Type NewDisplayType );
 
-	/**
-	* Handles checking if a display type is active for the graphical results window
-	*/
+	/** Handles checking if a display type is active for the graphical results window. */
 	ESlateCheckBoxState::Type HandleResultDisplayTypeIsChecked( EAutomationGrapicalDisplayType::Type InDisplayType ) const;
 
 	/** 
-	 * Gets the visibility for the throbber
+	 * Gets the visibility for the throbber.
 	 *
-	 * @return Whether we should show the throbber or not
+	 * @return Whether we should show the throbber or not.
 	 */
 	EVisibility GetTestsUpdatingThrobberVisibility() const
 	{
@@ -482,91 +477,93 @@ private:
 
 private:
 
-	// The automation window actions list
+	/** The automation window actions list. */
 	TSharedPtr<FUICommandList> AutomationWindowActions;
 
-	// Holds a pointer to the active session.
+	/** Holds a pointer to the active session. */
 	ISessionInfoPtr ActiveSession;
 
-	// Holds the AutomationController
+	/** Holds the AutomationController. */
 	IAutomationControllerManagerPtr AutomationController;
 
-	// Holds the search box widget.
+	/** Holds the search box widget. */
 	TSharedPtr< SSearchBox > AutomationSearchBox;
 
-	// Must maintain a widget size so the header and row icons can line up.
+	/** Must maintain a widget size so the header and row icons can line up. */
 	const float ColumnWidth;
 
-	// Global checkbox to enable/disable all visible tests.
+	/** Global checkbox to enable/disable all visible tests. */
 	TSharedPtr< SCheckBox > HeaderCheckbox;
 
-	// The list of all valid tests.
+	/** The list of all valid tests. */
 	TSharedPtr< SAutomationTestTreeView< TSharedPtr< IAutomationReport > > > TestTable;
 	
-	// Widget for header platform icons.
+	/** Widget for header platform icons. */
 	TSharedPtr< SHorizontalBox > PlatformsHBox;
 
-	// Widget for the command bar.
+	/** Widget for the command bar. */
 	TSharedPtr< SAutomationWindowCommandBar > CommandBar;
 
-	// Widget for the menu bar - run automation etc.
+	/** Widget for the menu bar - run automation etc. */
 	TSharedPtr< SVerticalBox > MenuBar;
 
-	// Holds the widget to display log messages
+	/** Holds the widget to display log messages. */
 	TSharedPtr<SListView<TSharedPtr<FAutomationOutputMessage> > > LogListView;
 
-	// Holds the widget to display a graph of the results
+	/** Holds the widget to display a graph of the results. */
 	TSharedPtr< SAutomationGraphicalResultBox > GraphicalResultBox;
 
-	// Holds the collection of log messages.
+	/** Holds the collection of log messages. */
 	TArray<TSharedPtr<FAutomationOutputMessage> > LogMessages;
 
-	// The automation report text filter.
+	/** The automation report text filter. */
 	TSharedPtr< AutomationReportTextFilter > AutomationTextFilter;
 
-	// The automation general filter - for smoke tests / warnings and Errors.
+	/** The automation general filter - for smoke tests / warnings and Errors. */
 	TSharedPtr< FAutomationFilter > AutomationGeneralFilter;
 
-	// The automation filter collection - contains the automation filters.
+	/** The automation filter collection - contains the automation filters. */
 	TSharedPtr< AutomationFilterCollection > AutomationFilters;
 
-	// Holds the session manager.
+	/** Holds the session manager. */
 	ISessionManagerPtr SessionManager;
 
-	// Holds the automation controller module state.
-	// This is set by the automation controller callback. We may go back to querying the module directly.
+	/**
+	 * Holds the automation controller module state.
+	 * This is set by the automation controller callback. We may go back to querying the module directly.
+	 */
 	EAutomationControllerModuleState::Type AutomationControllerState;
 
-	// Flag to acknowledge if the window is awaiting tests to display
+	/** Flag to acknowledge if the window is awaiting tests to display. */
 	bool bIsRequestingTests;
 
-	// Flag to tell if we have a child test selected in the test tree */
+	/** Flag to tell if we have a child test selected in the test tree. */
 	bool bHasChildTestSelected;
 
-	// Which type of window style to use for the test background
+	/** Which type of window style to use for the test background. */
 	EAutomationTestBackgroundStyle::Type TestBackgroundType;
 
-	// True if we are creating a new preset (The add preset text box is visible)
+	/** True if we are creating a new preset (The add preset text box is visible). */
 	bool bAddingTestPreset;
 
-	// Holds a pointer to the preset manager
+	/** Holds a pointer to the preset manager. */
 	TSharedPtr<FAutomationTestPresetManager> TestPresetManager;
 
-	// Holds the currently selected preset
+	/** Holds the currently selected preset. */
 	TSharedPtr<FAutomationTestPreset> SelectedPreset;
 
-	// Holds a pointer to the preset combo box widget
+	/** Holds a pointer to the preset combo box widget. */
 	TSharedPtr< SComboBox< TSharedPtr<FAutomationTestPreset> > > PresetComboBox;
 
-	// Holds a pointer to the preset text box
+	/** Holds a pointer to the preset text box. */
 	TSharedPtr<SEditableTextBox> PresetTextBox;
 
-	// Hold a pointer to the test tables header row.
+	/** Hold a pointer to the test tables header row. */
 	TSharedPtr<SHeaderRow> TestTableHeaderRow;
 
-	// Flag to determine whether we are tracking test history
+	/** Flag to determine whether we are tracking test history */
 	bool bIsTrackingHistory;
 
-	// The number of history elements we are tracking in this session
+	/** The number of history elements we are tracking in this session */
 	int32 NumHistoryElementsToTrack;
-};
\ No newline at end of file
+};
diff --git a/Engine/Source/Developer/AutomationWindow/Private/SAutomationWindowCommandBar.cpp b/Engine/Source/Developer/AutomationWindow/Private/SAutomationWindowCommandBar.cpp
--- a/Engine/Source/Developer/AutomationWindow/Private/SAutomationWindowCommandBar.cpp
+++ b/Engine/Source/Developer/AutomationWindow/Private/SAutomationWindowCommandBar.cpp
@@ -1,9 +1,5 @@
 // Copyright 1998-2014 Epic Games, Inc. All Rights Reserved.
 
-/*=============================================================================
-	SAutomationWindowCommandBar.cpp: Implements the SAutomationWindowCommandBar class.
-=============================================================================*/
-
 #include "AutomationWindowPrivatePCH.h"
 
 #define LOCTEXT_NAMESPACE "SAutomationWindowCommandBar"
@@ -63,4 +59,4 @@ void SAutomationWindowCommandBar::SetNumLogMessages( int32 Count )
 }
 
 
-#undef LOCTEXT_NAMESPACE
\ No newline at end of file
+#undef LOCTEXT_NAMESPACE
diff --git a/Engine/Source/Developer/AutomationWindow/Private/SAutomationWindowCommandBar.h b/Engine/Source/Developer/AutomationWindow/Private/SAutomationWindowCommandBar.h
--- a/Engine/Source/Developer/AutomationWindow/Private/SAutomationWindowCommandBar.h
+++ b/Engine/Source/Developer/AutomationWindow/Private/SAutomationWindowCommandBar.h
@@ -1,9 +1,5 @@
 // Copyright 1998-2014 Epic Games, Inc. All Rights Reserved.
 
-/*=============================================================================
-	SAutomationWindowCommandBar.h: Declares the SAutomationWindowCommandBar class.
-=============================================================================*/
-
 #pragma once
 
 
@@ -17,9 +13,7 @@ public:
 
 	SLATE_BEGIN_ARGS(SAutomationWindowCommandBar) { }
 
- 		/**
- 		 * Called when the copy log button is clicked.
- 		 */
+ 		/** Called when the copy log button is clicked. */
  		SLATE_EVENT(FOnClicked, OnCopyLogClicked)
 
 	SLATE_END_ARGS()
@@ -29,29 +23,27 @@ public:
 	/**
 	 * Construct this widget
 	 *
-	 * @param InArgs - The declaration data for this widget.
+	 * @param InArgs The declaration data for this widget.
 	 */
 	void Construct( const FArguments& InArgs, const TSharedRef< SNotificationList >& InNotificationList );
 
 	/**
 	 * Sets the number of messages selected in the log window.
 	 *
-	 * @param Count - Number of selected messages.
+	 * @param Count Number of selected messages.
 	 */
 	void SetNumLogMessages( int Count );
 
 private:
 
-	// Handles clicking the copy log button.
+	/** Handles clicking the copy log button. */
 	FReply HandleCopyButtonClicked( );
 
 private:
 
- 	// Holds the copy log button.
+ 	/** Holds the copy log button. */
  	TSharedPtr<SButton> CopyButton;
 
-private:
-
- 	// Holds a delegate that is executed when the copy log button is clicked.
+ 	/** Holds a delegate that is executed when the copy log button is clicked. */
  	FOnClicked OnCopyLogClicked;
 };
diff --git a/Engine/Source/Developer/AutomationWindow/Public/AutomationWindow.h b/Engine/Source/Developer/AutomationWindow/Public/AutomationWindow.h
--- a/Engine/Source/Developer/AutomationWindow/Public/AutomationWindow.h
+++ b/Engine/Source/Developer/AutomationWindow/Public/AutomationWindow.h
@@ -1,9 +1,5 @@
 // Copyright 1998-2014 Epic Games, Inc. All Rights Reserved.
 
-/*=============================================================================
-	AutomationWindow.h: AutomationWindow module public header file.
-=============================================================================*/
-
 #pragma once
 
 
@@ -18,4 +14,5 @@
 
 /* Interfaces
  *****************************************************************************/
+
 #include "IAutomationWindowModule.h"
diff --git a/Engine/Source/Developer/AutomationWindow/Public/Interfaces/IAutomationWindowModule.h b/Engine/Source/Developer/AutomationWindow/Public/Interfaces/IAutomationWindowModule.h
--- a/Engine/Source/Developer/AutomationWindow/Public/Interfaces/IAutomationWindowModule.h
+++ b/Engine/Source/Developer/AutomationWindow/Public/Interfaces/IAutomationWindowModule.h
@@ -1,8 +1,5 @@
 // Copyright 1998-2014 Epic Games, Inc. All Rights Reserved.
 
-/*=============================================================================
-	IAutomationWindowModule.h: Declares the IAutomationWindowModule interface.
-=============================================================================*/
 #pragma once
 
 
@@ -19,26 +16,25 @@ class IAutomationWindowModule
 public:
 
 	/**
-	 * Creates a Automation Window widget
-	 *
-	 * @param AutomationController - The automation controller to use
-	 * @param SessionManager - The session manager to use.
+	 * Creates a Automation Window widget.
 	 *
+	 * @param AutomationController The automation controller to use.
+	 * @param SessionManager The session manager to use.
 	 * @return	New Automation Test window widget
 	 */
 	virtual TSharedRef<class SWidget> CreateAutomationWindow( const IAutomationControllerManagerRef& AutomationController, const ISessionManagerRef& SessionManager ) = 0;
 
 	/**
-	 * Store a pointer to the AutomationWindowTab
+	 * Store a pointer to the AutomationWindowTab.
 	 *
-	 * @return a weak pointer to the automation tab
+	 * @return a weak pointer to the automation tab.
 	 */
 	virtual TWeakPtr<class SDockTab> GetAutomationWindowTab()  = 0;
 
 	/**
-	 * Store a pointer to the AutomationWindowTab
+	 * Store a pointer to the AutomationWindowTab.
 	 *
-	 * @param AutomationWindowTab a weak pointer to the automation tab
+	 * @param AutomationWindowTab a weak pointer to the automation tab.
 	 */
 	virtual void SetAutomationWindowTab(TWeakPtr<class SDockTab> AutomationWindowTab) = 0;
 
@@ -53,8 +49,6 @@ public:
 
 protected:
 
-	/**
-	 * Virtual destructor.
-	 */
+	/** Virtual destructor. */
 	~IAutomationWindowModule( ) { }
 };
diff --git a/Engine/Source/Developer/DesktopPlatform/Private/Mac/MacNativeFeedbackContext.h b/Engine/Source/Developer/DesktopPlatform/Private/Mac/MacNativeFeedbackContext.h
--- a/Engine/Source/Developer/DesktopPlatform/Private/Mac/MacNativeFeedbackContext.h
+++ b/Engine/Source/Developer/DesktopPlatform/Private/Mac/MacNativeFeedbackContext.h
@@ -1,11 +1,8 @@
 // Copyright 1998-2014 Epic Games, Inc. All Rights Reserved.
 
-/*=============================================================================
-	MacNativeFeedbackContext.h: Unreal Mac user interface interaction.
-=============================================================================*/
-
 #pragma once
 
+
 @interface FMacNativeFeedbackContextWindowController : NSObject
 {
 @public
diff --git a/Engine/Source/Developer/DesktopPlatform/Private/Windows/WindowsNativeFeedbackContext.h b/Engine/Source/Developer/DesktopPlatform/Private/Windows/WindowsNativeFeedbackContext.h
--- a/Engine/Source/Developer/DesktopPlatform/Private/Windows/WindowsNativeFeedbackContext.h
+++ b/Engine/Source/Developer/DesktopPlatform/Private/Windows/WindowsNativeFeedbackContext.h
@@ -1,14 +1,11 @@
 // Copyright 1998-2014 Epic Games, Inc. All Rights Reserved.
 
-/*=============================================================================
-	WindowsNativeFeedbackContext.h: Unreal Windows user interface interaction.
-=============================================================================*/
-
 #pragma once
 
 #include "AllowWindowsPlatformTypes.h"
 #include <CommCtrl.h>
 
+
 /**
  * Feedback context implementation for windows.
  */
diff --git a/Engine/Source/Developer/FriendsAndChat/Classes/FriendsAndChatStyle.h b/Engine/Source/Developer/FriendsAndChat/Classes/FriendsAndChatStyle.h
--- a/Engine/Source/Developer/FriendsAndChat/Classes/FriendsAndChatStyle.h
+++ b/Engine/Source/Developer/FriendsAndChat/Classes/FriendsAndChatStyle.h
@@ -1,12 +1,10 @@
 // Copyright 1998-2014 Epic Games, Inc. All Rights Reserved.
 
-/*=============================================================================
-	FriendsAndChatStyle.h: Declares the FFriendsAndChatStyle interface.
-=============================================================================*/
 #pragma once
 
 #include "FriendsAndChatStyle.generated.h"
 
+
 /**
  * Interface for the services manager.
  */
@@ -17,16 +15,15 @@ struct FRIENDSANDCHAT_API FFriendsAndChatStyle
 	GENERATED_USTRUCT_BODY()
 
 	// Default Constructor
-	FFriendsAndChatStyle(){}
+	FFriendsAndChatStyle() { }
 
 	// Default Destructor
-	virtual ~FFriendsAndChatStyle() {}
+	virtual ~FFriendsAndChatStyle() { }
 
 	/**
 	 * Override widget style function.
 	 */
-	virtual void GetResources( TArray< const FSlateBrush* >& OutBrushes ) const override
-	{}
+	virtual void GetResources( TArray< const FSlateBrush* >& OutBrushes ) const override { }
 
 	// Holds the widget type name
 	static const FName TypeName;
diff --git a/Engine/Source/Developer/FriendsAndChat/Private/FriendsAndChatManager.cpp b/Engine/Source/Developer/FriendsAndChat/Private/FriendsAndChatManager.cpp
--- a/Engine/Source/Developer/FriendsAndChat/Private/FriendsAndChatManager.cpp
+++ b/Engine/Source/Developer/FriendsAndChat/Private/FriendsAndChatManager.cpp
@@ -1,13 +1,11 @@
 // Copyright 1998-2014 Epic Games, Inc. All Rights Reserved.
 
-/*=============================================================================
-	FriendsAndChatManager.cpp: Implements the FFriendsAndChatManager class.
-=============================================================================*/
-
 #include "FriendsAndChatPrivatePCH.h"
 
+
 #define LOCTEXT_NAMESPACE "FriendsAndChatManager"
 
+
 /* FFriendsAndChatManager structors
  *****************************************************************************/
 
@@ -18,10 +16,12 @@ FFriendsAndChatManager::FFriendsAndChatManager( )
 	, ConfirmedFriendsCount( 0 )
 	, bIsInSession( false )
 	, bIsInited( false )
-{}
+{ }
+
 
 FFriendsAndChatManager::~FFriendsAndChatManager( )
-{}
+{ }
+
 
 /* IFriendsAndChatManager interface
  *****************************************************************************/
@@ -30,6 +30,7 @@ void FFriendsAndChatManager::StartupManager()
 {
 }
 
+
 void FFriendsAndChatManager::Init( FOnFriendsNotification& NotificationDelegate )
 {
 	// Clear existing data
@@ -42,7 +43,7 @@ void FFriendsAndChatManager::Init( FOnFriendsNotification& NotificationDelegate
 
 	OnlineSubMcp = static_cast< FOnlineSubsystemMcp* >( IOnlineSubsystem::Get( TEXT( "MCP" ) ) );
 
-	if (OnlineSubMcp != NULL &&
+	if (OnlineSubMcp != nullptr &&
 		OnlineSubMcp->GetMcpAccountMappingService().IsValid() &&
 		OnlineSubMcp->GetIdentityInterface().IsValid())
 	{
@@ -80,6 +81,7 @@ void FFriendsAndChatManager::Init( FOnFriendsNotification& NotificationDelegate
 	}
 }
 
+
 bool FFriendsAndChatManager::Tick( float Delta )
 {
 	if ( ManagerState == EFriendsAndManagerState::Idle )
@@ -111,6 +113,7 @@ bool FFriendsAndChatManager::Tick( float Delta )
 	return true;
 }
 
+
 void FFriendsAndChatManager::Logout()
 {
 	if ( FriendsInterface.IsValid() )
@@ -121,7 +124,7 @@ void FFriendsAndChatManager::Logout()
 		FriendsInterface->ClearOnDeleteFriendsListCompleteDelegate( 0, OnDeleteFriendsListCompleteDelegate );
 		FriendsInterface->ClearOnDeleteFriendCompleteDelegate( 0, OnDeleteFriendCompleteDelegate );
 
-		if ( OnlineSubMcp != NULL )
+		if ( OnlineSubMcp != nullptr )
 		{
 			FOnlineAccountMappingMcpPtr OnlineAccountMappingMcp = OnlineSubMcp->GetMcpAccountMappingService();
 			if ( OnlineAccountMappingMcp.IsValid() )
@@ -146,22 +149,23 @@ void FFriendsAndChatManager::Logout()
 	if ( FriendWindow.IsValid() )
 	{
 		FriendWindow->RequestDestroyWindow();
-		FriendWindow = NULL;
+		FriendWindow = nullptr;
 	}
 
 	if ( ChatWindow.IsValid() )
 	{
 		ChatWindow->RequestDestroyWindow();
-		ChatWindow = NULL;
+		ChatWindow = nullptr;
 	}
 
-	OnlineSubMcp = NULL;
+	OnlineSubMcp = nullptr;
 	if ( UpdateFriendsTickerDelegate.IsBound() )
 	{
 		FTicker::GetCoreTicker().RemoveTicker( UpdateFriendsTickerDelegate );
 	}
 }
 
+
 void FFriendsAndChatManager::GenerateFriendsWindow( TSharedPtr< const SWidget > InParentWidget, const FFriendsAndChatStyle* InStyle )
 {
 	const FVector2D DEFAULT_WINDOW_SIZE = FVector2D(308, 458);
@@ -205,17 +209,20 @@ void FFriendsAndChatManager::GenerateFriendsWindow( TSharedPtr< const SWidget >
 	}
 }
 
+
 void FFriendsAndChatManager::SetInSession( bool bInSession )
 {
 	bIsInSession = bInSession;
 	ReadListRequests.AddUnique( EFriendsLists::RecentPlayers );
 }
 
+
 bool FFriendsAndChatManager::IsInSession()
 {
 	return bIsInSession;
 }
 
+
 TSharedPtr< SWidget > FFriendsAndChatManager::GenerateFriendsListWidget( const FFriendsAndChatStyle* InStyle )
 {
 	if ( !FriendListWidget.IsValid() )
@@ -228,6 +235,7 @@ TSharedPtr< SWidget > FFriendsAndChatManager::GenerateFriendsListWidget( const F
 	return FriendListWidget;
 }
 
+
 void FFriendsAndChatManager::GenerateChatWindow( TSharedPtr< FFriendStuct > FriendItem )
 {
 	const FVector2D DEFAULT_WINDOW_SIZE = FVector2D(308, 458);
@@ -266,6 +274,7 @@ void FFriendsAndChatManager::GenerateChatWindow( TSharedPtr< FFriendStuct > Frie
 	}
 }
 
+
 void FFriendsAndChatManager::AcceptFriend( TSharedPtr< FFriendStuct > FriendItem )
 {
 	PendingOutgoingAcceptFriendRequests.Add( FriendItem->GetOnlineUser()->GetUserId() );
@@ -275,6 +284,7 @@ void FFriendsAndChatManager::AcceptFriend( TSharedPtr< FFriendStuct > FriendItem
 	ClearNotification( FriendItem );
 }
 
+
 void FFriendsAndChatManager::RejectFriend( TSharedPtr< FFriendStuct > FriendItem )
 {
 	NotifiedRequest.Remove( FriendItem->GetOnlineUser()->GetUserId() );
@@ -285,6 +295,7 @@ void FFriendsAndChatManager::RejectFriend( TSharedPtr< FFriendStuct > FriendItem
 	ClearNotification( FriendItem );
 }
 
+
 void FFriendsAndChatManager::ClearNotification( TSharedPtr< FFriendStuct > FriendItem )
 {
 	if ( NotificationMessages.Num() > 0 )
@@ -299,6 +310,7 @@ void FFriendsAndChatManager::ClearNotification( TSharedPtr< FFriendStuct > Frien
 	}
 }
 
+
 FReply FFriendsAndChatManager::HandleMessageAccepted( TSharedPtr< FFriendsAndChatMessage > ChatMessage, EFriendsResponseType::Type ResponseType )
 {
 	switch ( ResponseType )
@@ -335,6 +347,7 @@ FReply FFriendsAndChatManager::HandleMessageAccepted( TSharedPtr< FFriendsAndCha
 	return FReply::Handled();
 }
 
+
 /** Functor for comparing friends list */
 struct FCompareGroupByName
 {
@@ -346,6 +359,7 @@ struct FCompareGroupByName
 	}
 };
 
+
 bool FFriendsAndChatManager::UpdateFriendsList()
 {
 	bool bChanged = false;
@@ -403,6 +417,7 @@ bool FFriendsAndChatManager::UpdateFriendsList()
 	return bChanged;
 }
 
+
 bool FFriendsAndChatManager::UpdateRecentPlayerList()
 {
 	bool bChanged = false;
@@ -432,23 +447,27 @@ bool FFriendsAndChatManager::UpdateRecentPlayerList()
 	return bChanged;
 }
 
+
 int32 FFriendsAndChatManager::GetFilteredFriendsList( TArray< TSharedPtr< FFriendStuct > >& OutFriendsList )
 {
 	OutFriendsList = FilteredFriendsList;
 	return OutFriendsList.Num();
 }
 
+
 int32 FFriendsAndChatManager::GetFilteredOutgoingFriendsList( TArray< TSharedPtr< FFriendStuct > >& OutFriendsList )
 {
 	OutFriendsList = FilteredOutgoingList;
 	return OutFriendsList.Num();
 }
 
+
 int32 FFriendsAndChatManager::GetFriendCount()
 {
 	return ConfirmedFriendsCount;
 }
 
+
 void FFriendsAndChatManager::GenerateFriendsCount()
 {
 	ConfirmedFriendsCount = 0;
@@ -461,6 +480,7 @@ void FFriendsAndChatManager::GenerateFriendsCount()
 	}
 }
 
+
 void FFriendsAndChatManager::DeleteFriend( TSharedPtr< FFriendStuct > FriendItem )
 {
 	TSharedRef<FUniqueNetId> UserID = FriendItem->GetOnlineFriend()->GetUserId();
@@ -471,16 +491,18 @@ void FFriendsAndChatManager::DeleteFriend( TSharedPtr< FFriendStuct > FriendItem
 	RefreshList();
 }
 
+
 FReply FFriendsAndChatManager::OnCloseClicked()
 {
 	if ( FriendWindow.IsValid() )
 	{
 		FriendWindow->RequestDestroyWindow();
-		FriendWindow = NULL;
+		FriendWindow = nullptr;
 	}
 	return FReply::Handled();
 }
 
+
 FReply FFriendsAndChatManager::OnMinimizeClicked()
 {
 	if ( FriendWindow.IsValid() )
@@ -490,16 +512,18 @@ FReply FFriendsAndChatManager::OnMinimizeClicked()
 	return FReply::Handled();
 }
 
+
 FReply FFriendsAndChatManager::OnChatCloseClicked()
 {
 	if ( ChatWindow.IsValid() )
 	{
 		ChatWindow->RequestDestroyWindow();
-		ChatWindow = NULL;
+		ChatWindow = nullptr;
 	}
 	return FReply::Handled();
 }
 
+
 FReply FFriendsAndChatManager::OnChatMinimizeClicked()
 {
 	if ( ChatWindow.IsValid() )
@@ -509,6 +533,7 @@ FReply FFriendsAndChatManager::OnChatMinimizeClicked()
 	return FReply::Handled();
 }
 
+
 void FFriendsAndChatManager::RefreshList()
 {
 	switch ( CurrentList )
@@ -532,6 +557,7 @@ void FFriendsAndChatManager::RefreshList()
 	}
 }
 
+
 void FFriendsAndChatManager::BuildFriendsList()
 {
 	FilteredFriendsList.Empty();
@@ -548,6 +574,7 @@ void FFriendsAndChatManager::BuildFriendsList()
 	OnFriendsListUpdatedDelegate.Broadcast();
 }
 
+
 void FFriendsAndChatManager::BuildRecentPlayersList()
 {
 	FilteredFriendsList.Empty();
@@ -561,6 +588,7 @@ void FFriendsAndChatManager::BuildRecentPlayersList()
 	OnFriendsListUpdatedDelegate.Broadcast();
 }
 
+
 void FFriendsAndChatManager::BuildRequestIncomingPlayersList()
 {
 	FilteredFriendsList.Empty();
@@ -575,6 +603,7 @@ void FFriendsAndChatManager::BuildRequestIncomingPlayersList()
 	}
 }
 
+
 void FFriendsAndChatManager::BuildRequestOutgoingPlayersList()
 {
 	FilteredOutgoingList.Empty();
@@ -609,6 +638,7 @@ void FFriendsAndChatManager::BuildRequestOutgoingPlayersList()
 	OnFriendsListUpdatedDelegate.Broadcast();
 }
 
+
 void FFriendsAndChatManager::RequestFriend( const FText& FriendName )
 {
 	if ( !FriendName.IsEmpty() )
@@ -621,6 +651,7 @@ void FFriendsAndChatManager::RequestFriend( const FText& FriendName )
 	}
 }
 
+
 void FFriendsAndChatManager::SendFriendRequests()
 {
 	// Invite Friends
@@ -635,11 +666,13 @@ void FFriendsAndChatManager::SendFriendRequests()
 	}
 }
 
+
 EFriendsAndManagerState::Type FFriendsAndChatManager::GetManagerState()
 {
 	return ManagerState;
 }
 
+
 void FFriendsAndChatManager::SetListSelect( EFriendsDisplayLists::Type ListSelectType )
 {
 	if ( CurrentList != ListSelectType )
@@ -649,6 +682,7 @@ void FFriendsAndChatManager::SetListSelect( EFriendsDisplayLists::Type ListSelec
 	}
 }
 
+
 TSharedPtr< FUniqueNetId > FFriendsAndChatManager::FindUserID( const FString& InUsername )
 {
 	for ( int32 Index = 0; Index < FriendsList.Num(); Index++ )
@@ -658,14 +692,16 @@ TSharedPtr< FUniqueNetId > FFriendsAndChatManager::FindUserID( const FString& In
 			return FriendsList[ Index ]->GetUniqueID();
 		}
 	}
-	return NULL;
+	return nullptr;
 }
 
+
 bool FFriendsAndChatManager::IsPendingInvite( const FString& InUsername )
 {
 	return FriendByNameRequests.Contains( InUsername ) || FriendByNameInvites.Contains( InUsername );
 }
 
+
 TSharedPtr< FFriendStuct > FFriendsAndChatManager::FindUser( TSharedRef<FUniqueNetId> InUserID, const FString& ListName )
 {
 	TArray< TSharedPtr< FFriendStuct > >& SearchList = ListName == EFriendsLists::ToString( EFriendsLists::Default ) ? FriendsList : RecentPlayerList;
@@ -678,9 +714,10 @@ TSharedPtr< FFriendStuct > FFriendsAndChatManager::FindUser( TSharedRef<FUniqueN
 		}
 	}
 
-	return NULL;
+	return nullptr;
 }
 
+
 void FFriendsAndChatManager::SetState( EFriendsAndManagerState::Type NewState )
 {
 	ManagerState = NewState;
@@ -760,6 +797,7 @@ void FFriendsAndChatManager::SetState( EFriendsAndManagerState::Type NewState )
 	}
 }
 
+
 void FFriendsAndChatManager::SendFriendInviteNotification()
 {
 	check( PendingIncomingInvitesList[ 0 ].IsValid() )
@@ -781,9 +819,10 @@ void FFriendsAndChatManager::SendFriendInviteNotification()
 	PendingIncomingInvitesList.RemoveAt( 0 );
 }
 
+
 void FFriendsAndChatManager::OnQueryUserIdFromDisplayNameComplete(bool bWasSuccessful, const FUniqueNetId& RequestingUserId, const FString& DisplayName, const FUniqueNetId& IdentifiedUserId, const FString& Error)
 {
-	check( OnlineSubMcp != NULL && OnlineSubMcp->GetMcpAccountMappingService().IsValid() );
+	check( OnlineSubMcp != nullptr && OnlineSubMcp->GetMcpAccountMappingService().IsValid() );
 
 	if ( bWasSuccessful && IdentifiedUserId.IsValid() )
 	{
@@ -818,6 +857,7 @@ void FFriendsAndChatManager::OnQueryUserIdFromDisplayNameComplete(bool bWasSucce
 	}
 }
 
+
 void FFriendsAndChatManager::OnSendInviteComplete( int32 LocalPlayer, bool bWasSuccessful, const FUniqueNetId& FriendId, const FString& ListName, const FString& ErrorStr )
 {
 	PendingOutgoingFriendRequests.RemoveAt( 0 );
@@ -829,6 +869,7 @@ void FFriendsAndChatManager::OnSendInviteComplete( int32 LocalPlayer, bool bWasS
 	}
 }
 
+
 void FFriendsAndChatManager::OnReadFriendsComplete( int32 LocalPlayer, bool bWasSuccessful, const FString& ListName, const FString& ErrorStr )
 {
 	TArray< TSharedRef<FOnlineFriend> > Friends;
@@ -864,7 +905,7 @@ void FFriendsAndChatManager::OnReadFriendsComplete( int32 LocalPlayer, bool bWas
 			}
 		}
 
-		check( OnlineSubMcp != NULL && OnlineSubMcp->GetMcpAccountMappingService().IsValid() );
+		check( OnlineSubMcp != nullptr && OnlineSubMcp->GetMcpAccountMappingService().IsValid() );
 		IOnlineUserPtr UserInterface = OnlineSubMcp->GetUserInterface();
 
 		if ( QueryUserIds.Num() > 0 )
@@ -893,9 +934,10 @@ void FFriendsAndChatManager::OnReadFriendsComplete( int32 LocalPlayer, bool bWas
 	}
 }
 
+
 void FFriendsAndChatManager::OnQueryUserInfoComplete( int32 LocalPlayer, bool bWasSuccessful, const TArray< TSharedRef<class FUniqueNetId> >& UserIds, const FString& ErrorStr )
 {
-	check(OnlineSubMcp != NULL && OnlineSubMcp->GetMcpAccountMappingService().IsValid());
+	check(OnlineSubMcp != nullptr && OnlineSubMcp->GetMcpAccountMappingService().IsValid());
 	IOnlineUserPtr UserInterface = OnlineSubMcp->GetUserInterface();
 
 	OnlineSubMcp->GetUserInterface()->ClearOnQueryUserInfoCompleteDelegate(0, OnQueryUserInfoCompleteDelegate);
@@ -941,6 +983,7 @@ void FFriendsAndChatManager::OnQueryUserInfoComplete( int32 LocalPlayer, bool bW
 	}
 }
 
+
 void FFriendsAndChatManager::OnAcceptInviteComplete( int32 LocalPlayer, bool bWasSuccessful, const FUniqueNetId& FriendId, const FString& ListName, const FString& ErrorStr )
 {
 	PendingOutgoingAcceptFriendRequests.RemoveAt( 0 );
@@ -958,12 +1001,14 @@ void FFriendsAndChatManager::OnAcceptInviteComplete( int32 LocalPlayer, bool bWa
 	}
 }
 
+
 void FFriendsAndChatManager::OnDeleteFriendsListComplete( int32 LocalPlayer, bool bWasSuccessful, const FString& ListName, const FString& ErrorStr )
 {
 	FriendsInterface->ClearOnDeleteFriendsListCompleteDelegate( 0, OnDeleteFriendsListCompleteDelegate );
 	SetState( EFriendsAndManagerState::Idle );
 }
 
+
 void FFriendsAndChatManager::OnDeleteFriendComplete( int32 LocalPlayer, bool bWasSuccessful, const FUniqueNetId& DeletedFriendID, const FString& ListName, const FString& ErrorStr )
 {
 	PendingOutgoingDeleteFriendRequests.RemoveAt( 0 );
@@ -980,9 +1025,11 @@ void FFriendsAndChatManager::OnDeleteFriendComplete( int32 LocalPlayer, bool bWa
 	}
 }
 
+
 /* FFriendsAndChatManager system singletons
 *****************************************************************************/
-TSharedPtr< FFriendsAndChatManager > FFriendsAndChatManager::SingletonInstance = NULL;
+
+TSharedPtr< FFriendsAndChatManager > FFriendsAndChatManager::SingletonInstance = nullptr;
 
 TSharedRef< FFriendsAndChatManager > FFriendsAndChatManager::Get()
 {
@@ -994,9 +1041,11 @@ TSharedRef< FFriendsAndChatManager > FFriendsAndChatManager::Get()
 	return SingletonInstance.ToSharedRef();
 }
 
+
 void FFriendsAndChatManager::Shutdown()
 {
 	SingletonInstance.Reset();
 }
 
+
 #undef LOCTEXT_NAMESPACE
diff --git a/Engine/Source/Developer/FriendsAndChat/Private/FriendsAndChatManager.h b/Engine/Source/Developer/FriendsAndChat/Private/FriendsAndChatManager.h
--- a/Engine/Source/Developer/FriendsAndChat/Private/FriendsAndChatManager.h
+++ b/Engine/Source/Developer/FriendsAndChat/Private/FriendsAndChatManager.h
@@ -1,11 +1,8 @@
 // Copyright 1998-2014 Epic Games, Inc. All Rights Reserved.
 
-/*=============================================================================
-	FriendsAndChatManager.h: Declares the FFriendsAndChatManager interface.
-=============================================================================*/
-
 #pragma once
 
+
 // Enum holding the state of the Friends manager
 namespace EFriendsAndManagerState
 {
@@ -23,11 +20,11 @@ namespace EFriendsAndManagerState
 	};
 };
 
-/**
- * Delegate type for FriendsList updated.
- */
+
+/** Delegate type for FriendsList updated. */
 DECLARE_MULTICAST_DELEGATE(FOnFriendsUpdated)
 
+
 /**
  * Implement the Friend and Chat manager
  */
@@ -36,124 +33,129 @@ class FFriendsAndChatManager
 {
 public:
 
-	/**
-	 * Default constructor
-	 */
+	/** Default constructor. */
 	FFriendsAndChatManager( );
 	
-	/**
-	 * Default destructor
-	 */
+	/** Destructor. */
 	~FFriendsAndChatManager( );
 
 public:
 
-	/**
-	 * Default constructor
-	 */
+	/** */
 	void StartupManager();
 
 	/**
-	 * Init the manager
-	 * @param Notification delegate - the delegate for message notifications
+	 * Init the manager.
+	 *
+	 * @param Notification delegate The delegate for message notifications.
 	 */
 	void Init( FOnFriendsNotification& NotificationDelegate );
 
-	/**
-	 * Logout - close any Friends windows.
-	 */
+	/** Logout and close any Friends windows. */
 	void Logout();
 
 	/**
-	 * Create the friends list window
-	 * @param ParentWindow - the holding window
-	 * @param InStyle - the style used to create the widgets
+	 * Create the friends list window.
+	 *
+	 * @param ParentWindow The holding window.
+	 * @param InStyle The style used to create the widgets.
 	 */
 	void GenerateFriendsWindow( TSharedPtr< const SWidget > ParentWidget, const FFriendsAndChatStyle* InStyle );
 
 	/**
-	 * Set if the player is in a session
-	 * @param bInSession - If we are in a game session
+	 * Set if the player is in a session.
+	 *
+	 * @param bInSession Whether we are in a game session.
 	 */
 	void SetInSession( bool bInSession );
 
 	/**
-	 * Get if the player is in a session
-	 * @return True if we are in a game session
+	 * Get if the player is in a session.
+	 *
+	 * @return True if we are in a game session.
 	 */
 	bool IsInSession();
 
 	/**
 	 * Create the friends list widget without a holder
-	 * @param InStyle - the style used to create the widgets
+	 *
+	 * @param InStyle The style used to create the widgets.
 	 * @return The generated widget.
 	 */
 	TSharedPtr< SWidget > GenerateFriendsListWidget( const FFriendsAndChatStyle* InStyle );
 
 	/**
 	 * Create the friends list window.
-	 * @param FriendItem - The friend to start a chat with.
+	 *
+	 * @param FriendItem The friend to start a chat with.
 	 */
 	void GenerateChatWindow( TSharedPtr< FFriendStuct > FriendItem );
 
 	/**
-	 * Accept a friend request
-	 * @param FriendItem - The friend to accept
+	 * Accept a friend request.
+	 *
+	 * @param FriendItem The friend to accept.
 	 */
 	void AcceptFriend( TSharedPtr< FFriendStuct > FriendItem );
 
 	/**
-	 * Reject a friend request
-	 * @param FriendItem - The friend to reject.
+	 * Reject a friend request.
+	 *
+	 * @param FriendItem The friend to reject.
 	 */
 	void RejectFriend( TSharedPtr< FFriendStuct > FriendItem );
 
 	/**
 	 * Clear any existing notifications
-	 * @param FriendItem - The friend related to the notification message.
+	 *
+	 * @param FriendItem The friend related to the notification message.
 	 */
 	void ClearNotification( TSharedPtr< FFriendStuct > FriendItem );
 
 	/**
-	 * Get the friends filtered list of friends
-	 * @param OutFriendsList  - array of friends to fill in
-	 * @return the friend list count
+	 * Get the friends filtered list of friends.
+	 *
+	 * @param OutFriendsList  Array of friends to fill in.
+	 * @return the friend list count.
 	 */
 	int32 GetFilteredFriendsList( TArray< TSharedPtr< FFriendStuct > >& OutFriendsList );
 
 	/**
 	 * Get outgoing request list.
-	 * @param OutFriendsList  - array of friends to fill in
-	 * @return the friend list count
+	 *
+	 * @param OutFriendsList  Array of friends to fill in.
+	 * @return The friend list count.
 	 */
 	int32 GetFilteredOutgoingFriendsList( TArray< TSharedPtr< FFriendStuct > >& OutFriendsList );
 
 	/**
-	 * Get the friends count
-	 * @return The amount of friends
+	 * Gets the friends count.
+	 *
+	 * @return The amount of friends.
 	 */
 	int32 GetFriendCount();
 
-	/**
-	 * Calculate the new friends count after an update
-	 */
+	/** Calculates the new friends count after an update. */
 	void GenerateFriendsCount();
 
 	/**
-	 * Request a friend be added
-	 * @param FriendName - The friend name
+	 * Request a friend be added.
+	 *
+	 * @param FriendName The friend name.
 	 */
 	void RequestFriend( const FText& FriendName );
 
 	/**
-	 * Request a friend to be deleted
-	 * @param FriendItem - The friend item to delete
+	 * Request a friend to be deleted.
+	 *
+	 * @param FriendItem The friend item to delete.
 	 */
 	void DeleteFriend( TSharedPtr< FFriendStuct > FriendItem );
 
 	/**
 	 * Accessor for the Friends List updated delegate.
-	 * @return The delegate
+	 *
+	 * @return The delegate.
 	 */
 	FOnFriendsUpdated& OnFriendsListUpdated()
 	{
@@ -162,213 +164,201 @@ public:
 
 	/**
 	 * Get the manager state.
-	 * @return The manager state
+	 *
+	 * @return The manager state.
 	 */
 	EFriendsAndManagerState::Type GetManagerState();
 
 	/**
 	 * Get the manager state.
-	 * @param bListSelectType -The list select type.
+	 *
+	 * @param bListSelectType The list select type.
 	 */
 	void SetListSelect( EFriendsDisplayLists::Type ListSelectType );
 
 	/**
 	 * Find a user ID.
-	 * @param InUserName -The user name to find.
+	 *
+	 * @param InUserName The user name to find.
 	 * @return The unique ID.
 	 */
 	TSharedPtr< FUniqueNetId > FindUserID( const FString& InUsername );
 
 	/**
 	 * Is this user pending an invite.
-	 * @param InUserName -The user name to find.
+	 *
+	 * @param InUserName The user name to find.
 	 * @return True if pending an invite.
 	 */	
 	bool IsPendingInvite( const FString& InUsername );
 
 	/**
-	 * Find a user
-	 * @param InUserName -The user name to find.
-	 * @param ListName -The list name.
+	 * Find a user.
+	 *
+	 * @param InUserName The user name to find.
+	 * @param ListName The list name.
 	 * @return The Friend ID.
 	 */
 	TSharedPtr< FFriendStuct > FindUser( TSharedRef<FUniqueNetId> InUserID, const FString& ListName );
 
 private:
 
-	/**
-	 * Process the friends list to send out invites
-	 */
+	/** Process the friends list to send out invites */
 	void SendFriendRequests();
 
-	/**
-	 * Refresh the data lists
-	 */
+	/** Refresh the data lists */
 	void RefreshList();
 
-	/**
-	 * Build the friends list used in the UI.
-	 */
+	/** Build the friends list used in the UI. */
 	void BuildFriendsList();
 
-	/**
-	 * Build the recent player list used in the UI.
-	 */
+	/** Build the recent player list used in the UI. */
 	void BuildRecentPlayersList();
 
-	/**
-	 * Build the incomming request list
-	 */
+	/** Build the incomming request list */
 	void BuildRequestIncomingPlayersList();
 
-	/**
-	 * Build the outgoing request list
-	 */
+	/** Build the outgoing request list */
 	void BuildRequestOutgoingPlayersList();
 
 	/**
-	 * Set the manager state
-	 * @param NewState - The new manager state
+	 * Set the manager state.
+	 *
+	 * @param NewState The new manager state.
 	 */
 	void SetState( EFriendsAndManagerState::Type NewState );
 
-	/**
-	 * Send a friend invite notification
-	 */
+	/** Send a friend invite notification. */
 	void SendFriendInviteNotification();
 
-	/**
-	 * Called when singleton is released
-	 */
+	/** Called when singleton is released. */
 	void ShutdownManager();
 
 	/**
-	 * Delegate used when the friends read request has completed
+	 * Delegate used when the friends read request has completed.
 	 *
-	 * @param LocalPlayer		- the controller number of the associated user that made the request
-	 * @param bWasSuccessful	- true if the async action completed without error, false if there was an error
-	 * @param ListName			- name of the friends list that was operated on
-	 * @param ErrorStr			- string representing the error condition
+	 * @param LocalPlayer		The controller number of the associated user that made the request.
+	 * @param bWasSuccessful	true if the async action completed without error, false if there was an error.
+	 * @param ListName			Name of the friends list that was operated on.
+	 * @param ErrorStr			String representing the error condition.
 	 */
 	void OnReadFriendsComplete( int32 LocalPlayer, bool bWasSuccessful, const FString& ListName, const FString& ErrorStr );
 
 	/**
-	 * Delegate used when an invite accept request has completed
+	 * Delegate used when an invite accept request has completed.
 	 *
-	 * @param LocalPlayer		- the controller number of the associated user that made the request
-	 * @param bWasSuccessful	- true if the async action completed without error, false if there was an error
-	 * @param FriendId			- player that invited us
-	 * @param ListName			- name of the friends list that was operated on
-	 * @param ErrorStr			- string representing the error condition
+	 * @param LocalPlayer		The controller number of the associated user that made the request
+	 * @param bWasSuccessful	true if the async action completed without error, false if there was an error
+	 * @param FriendId			Player that invited us
+	 * @param ListName			Name of the friends list that was operated on
+	 * @param ErrorStr			String representing the error condition
 	 */
 	void OnAcceptInviteComplete( int32 LocalPlayer, bool bWasSuccessful, const FUniqueNetId& FriendId, const FString& ListName, const FString& ErrorStr );
 
 	/**
-	 * Delegate used when an invite send request has completed
+	 * Delegate used when an invite send request has completed.
 	 *
-	 * @param LocalPlayer		- the controller number of the associated user that made the request
-	 * @param bWasSuccessful	- true if the async action completed without error, false if there was an error
-	 * @param FriendId			- player that was invited
-	 * @param ListName			- name of the friends list that was operated on
-	 * @param ErrorStr			- string representing the error condition
+	 * @param LocalPlayer		The controller number of the associated user that made the request.
+	 * @param bWasSuccessful	true if the async action completed without error, false if there was an error.
+	 * @param FriendId			Player that was invited.
+	 * @param ListName			Name of the friends list that was operated on.
+	 * @param ErrorStr			String representing the error condition.
 	 */
 	void OnSendInviteComplete( int32 LocalPlayer, bool bWasSuccessful, const FUniqueNetId& FriendId, const FString& ListName, const FString& ErrorStr );
 
 	/**
-	 * Delegate used when the friends list delete request has completed
+	 * Delegate used when the friends list delete request has completed.
 	 *
-	 * @param LocalPlayer		- the controller number of the associated user that made the request
-	 * @param bWasSuccessful	- true if the async action completed without error, false if there was an error
-	 * @param ListName			- name of the friends list that was operated on
-	 * @param ErrorStr			- string representing the error condition
+	 * @param LocalPlayer		The controller number of the associated user that made the request.
+	 * @param bWasSuccessful	true if the async action completed without error, false if there was an error.
+	 * @param ListName			Name of the friends list that was operated on.
+	 * @param ErrorStr			String representing the error condition.
      */
 	void OnDeleteFriendsListComplete( int32 LocalPlayer, bool bWasSuccessful, const FString& ListName, const FString& ErrorStr );
 
 	/**
-	 * Delegate used when the friends delete request has completed
+	 * Delegate used when the friends delete request has completed.
 	 *
-	 * @param LocalPlayer		- the controller number of the associated user that made the request
-	 * @param bWasSuccessful	- true if the async action completed without error, false if there was an error
-	 * @param DeletedFriendID	- the ID of the deleted friend
-	 * @param ListName			- name of the friends list that was operated on
-	 * @param ErrorStr			- string representing the error condition
+	 * @param LocalPlayer		The controller number of the associated user that made the request.
+	 * @param bWasSuccessful	true if the async action completed without error, false if there was an error.
+	 * @param DeletedFriendID	The ID of the deleted friend.
+	 * @param ListName			Name of the friends list that was operated on.
+	 * @param ErrorStr			String representing the error condition.
 	 */
 	void OnDeleteFriendComplete( int32 LocalPlayer, bool bWasSuccessful, const FUniqueNetId& DeletedFriendID, const FString& ListName, const FString& ErrorStr );
 
 	/**
-	 * Delegate used when a friend ID request has completed
+	 * Delegate used when a friend ID request has completed.
 	 *
-	 * @param bWasSuccessful		- true if the async action completed without error, false if there was an error
-	 * @param RequestingUserId		- the ID of the user making the request
-	 * @param DisplayName	- The string ID of the friend found
-	 * @param IdentifiedUserId		- The net ID of the found friend
-	 * @param ErrorStr				- string representing the error condition
+	 * @param bWasSuccessful		true if the async action completed without error, false if there was an error.
+	 * @param RequestingUserId		The ID of the user making the request.
+	 * @param DisplayName			The string ID of the friend found.
+	 * @param IdentifiedUserId		The net ID of the found friend.
+	 * @param ErrorStr				String representing the error condition.
 	 */
 	void OnQueryUserIdFromDisplayNameComplete(bool bWasSuccessful, const FUniqueNetId& RequestingUserId, const FString& DisplayName, const FUniqueNetId& IdentifiedUserId, const FString& Error);
 
 	/**
-	 * Delegate used when a query user info request has completed
+	 * Delegate used when a query user info request has completed.
 	 *
-	 * @param LocalPlayer		- the controller number of the associated user that made the request
-	 * @param bWasSuccessful	- true if the async action completed without error, false if there was an error
-	 * @param UserIds An array	- of user IDs filled out with info
-	 * @param ErrorStr			- string representing the error condition
+	 * @param LocalPlayer		The controller number of the associated user that made the request.
+	 * @param bWasSuccessful	true if the async action completed without error, false if there was an error.
+	 * @param UserIds			An array of user IDs filled out with info.
+	 * @param ErrorStr			String representing the error condition.
 	 */
 	void OnQueryUserInfoComplete( int32 LocalPlayer, bool bWasSuccessful, const TArray< TSharedRef<class FUniqueNetId> >& UserIds, const FString& ErrorStr );
 
 	/**
-	 * Handle an accept message accepted from a notification
+	 * Handle an accept message accepted from a notification.
 	 *
-	 * @param MessageNotification 	- The message responded to
- 	 * @param ResponseTpe 	- The response from the message request
+	 * @param MessageNotification The message responded to.
+ 	 * @param ResponseTpe The response from the message request.
 	 */
 	FReply HandleMessageAccepted( TSharedPtr< FFriendsAndChatMessage > MessageNotification, EFriendsResponseType::Type ResponseType );
 
 	/**
 	 * Validate friend list after a refresh. Check if the list has changed.
-	 * @return True if the list has changes and the UI needs refreshing
+	 *
+	 * @return True if the list has changes and the UI needs refreshing.
+	 * @see UpdateRecentPlayerList
 	 */
 	bool UpdateFriendsList();
 
 	/**
 	 * Validate recent player list after a refresh. Check if the list has changed.
+	 *
 	 * @return True if the list has changes and the UI needs refreshing
+	 * @see UpdateFriendsList
 	 */
 	bool UpdateRecentPlayerList();
 
 	/**
 	 * A ticker used to perform updates on the main thread.
-	 * @param Delta		The tick delta
-	 * @return	true to continue ticking
+	 *
+	 * @param Delta The tick delta.
+	 * @return true to continue ticking.
 	 */
 	bool Tick( float Delta );
 
 	// UI functions
 
-	/**
-	 * Close button clicked
-	 */
+	/** Close button clicked. */
 	FReply OnCloseClicked();
 
-	/**
-	 * Minimize button clicked
-	 */
+	/** Minimize button clicked. */
 	FReply OnMinimizeClicked();
 
-	/**
-	 * Close button clicked
-	 */
+	/** Close button clicked. */
 	FReply OnChatCloseClicked();
 
-	/**
-	 * Minimize button clicked
-	 */
+	/** Minimize button clicked. */
 	FReply OnChatMinimizeClicked();
 
 private:
 
 	/* Lists used to keep track of Friends
 	*****************************************************************************/
+
 	// Holds the list of user names to query add as friends 
 	TArray< FString > FriendByNameRequests;
 	// Holds the list of user names to send invites to. 
@@ -402,6 +392,7 @@ private:
 
 	/* Delegates
 	*****************************************************************************/
+
 	// Holds the ticker delegate
 	FTickerDelegate UpdateFriendsTickerDelegate;
 	// Delegate when friend list has been read
@@ -425,6 +416,7 @@ private:
 
 	/* Identity stuff
 	*****************************************************************************/
+
 	// Holds the Online identity
 	IOnlineIdentityPtr OnlineIdentity;
 	// Holds the Friends Interface
@@ -434,6 +426,7 @@ private:
 
 	/* Manger state
 	*****************************************************************************/
+
 	// Holds the manager state
 	EFriendsAndManagerState::Type ManagerState;
 	// Holds the update interval
@@ -443,6 +436,7 @@ private:
 
 	/* UI
 	*****************************************************************************/
+
 	// Holds the parent window
 	TSharedPtr< const SWidget > ParentWidget;
 	// Holds the main Friends List window
@@ -462,11 +456,12 @@ private:
 	// Holds if the Friends list is inited
 	bool bIsInited;
 
-/* Here we have static access for the singleton
-*****************************************************************************/
 public:
+
 	static TSharedRef< FFriendsAndChatManager > Get();
 	static void Shutdown();
+
 private:
+
 	static TSharedPtr< FFriendsAndChatManager > SingletonInstance;
 };
diff --git a/Engine/Source/Developer/FriendsAndChat/Private/FriendsAndChatModule.cpp b/Engine/Source/Developer/FriendsAndChat/Private/FriendsAndChatModule.cpp
--- a/Engine/Source/Developer/FriendsAndChat/Private/FriendsAndChatModule.cpp
+++ b/Engine/Source/Developer/FriendsAndChat/Private/FriendsAndChatModule.cpp
@@ -1,11 +1,8 @@
 // Copyright 1998-2014 Epic Games, Inc. All Rights Reserved.
 
-/*=============================================================================
-	FriendsAndChatModule.cpp: Implements the FFriendsAndChatModule.
-=============================================================================*/
-
 #include "FriendsAndChatPrivatePCH.h"
 
+
 /**
  * Implements the FriendsAndChat module.
  */
@@ -14,7 +11,8 @@ class FFriendsAndChatModule
 {
 public:
 
-	// Begin IFriendsAndChatModule interface
+	// IFriendsAndChatModule interface
+
 	virtual void Init( bool bAllowJoinGame ) override
 	{
 		FFriendsAndChatManager::Get()->Init( FriendsListNotificationDelegate );
@@ -62,14 +60,12 @@ public:
 		FFriendsMessageManager::Get()->Logout();
 	}
 
-	// End IFriendsAndChatModule interface
-
 public:
 
-	// Begin IModuleInterface interface
+	// IModuleInterface interface
+
 	virtual void StartupModule() override
 	{
-		// Make sure the singletons are created
  		FFriendsAndChatManager::Get();
 	}
 
@@ -77,11 +73,12 @@ public:
 	{
 		FFriendsAndChatManager::Shutdown();
 	}
-	// End IModuleInterface interface
 
 private:
-	// Holds the Notification delegate
+
+	/** Holds the Notification delegate. */
 	FOnFriendsNotification FriendsListNotificationDelegate;
 };
 
+
 IMPLEMENT_MODULE( FFriendsAndChatModule, FriendsAndChat );
diff --git a/Engine/Source/Developer/FriendsAndChat/Private/FriendsAndChatPrivatePCH.h b/Engine/Source/Developer/FriendsAndChat/Private/FriendsAndChatPrivatePCH.h
--- a/Engine/Source/Developer/FriendsAndChat/Private/FriendsAndChatPrivatePCH.h
+++ b/Engine/Source/Developer/FriendsAndChat/Private/FriendsAndChatPrivatePCH.h
@@ -1,14 +1,11 @@
 // Copyright 1998-2014 Epic Games, Inc. All Rights Reserved.
 
-/*=============================================================================
-	FriendsAndChatPrivatePCH.h: Pre-compiled header file for the 
-	FriendsAndChat module.
-=============================================================================*/
-
 #pragma once
 
+
 /* Dependencies
  *****************************************************************************/
+
 #include "Slate.h"
 #include "Ticker.h"
 #include "SlateStyle.h"
@@ -22,14 +19,18 @@
 #include "OnlineMessageInterface.h"
 #include "FriendListItems.h"
 
+
 /* Private UI includes
 *****************************************************************************/
+
 #include "SFriendsItem.h"
 #include "SFriendsListContextMenu.h"
 #include "SFriendsList.h"
 #include "SChatWindow.h"
 
+
 /* Private Class includes
 *****************************************************************************/
+
 #include "FriendsAndChatManager.h"
 #include "FriendsMessageManager.h"
diff --git a/Engine/Source/Developer/FriendsAndChat/Private/FriendsAndChatStyle.cpp b/Engine/Source/Developer/FriendsAndChat/Private/FriendsAndChatStyle.cpp
--- a/Engine/Source/Developer/FriendsAndChat/Private/FriendsAndChatStyle.cpp
+++ b/Engine/Source/Developer/FriendsAndChat/Private/FriendsAndChatStyle.cpp
@@ -1,85 +1,95 @@
 // Copyright 1998-2014 Epic Games, Inc. All Rights Reserved.
 
-/*=============================================================================
-	FriendsAndChatStyle.cpp: Implements the FFriendsAndChatStyle class.
-=============================================================================*/
-
 #include "FriendsAndChatPrivatePCH.h"
 
+
 const FName FFriendsAndChatStyle::TypeName( TEXT("FFriendsAndChatStyle") );
 
+
 FFriendsAndChatStyle& FFriendsAndChatStyle::SetCloseButtonStyle( const FButtonStyle& InButtonStyle)
 {
 	CloseButtonStyle = InButtonStyle; 
 	return *this; 
 }
 
+
 FFriendsAndChatStyle& FFriendsAndChatStyle::SetMinimizeStyle( const FButtonStyle& InButtonStyle)
 {
 	MinimizeButtonStyle = InButtonStyle; 
 	return *this;
 }
 
+
 FFriendsAndChatStyle& FFriendsAndChatStyle::SetSearchButtonStyle( const FButtonStyle& InButtonStyle)
 {
 	SearchButtonStyle = InButtonStyle; 
 	return *this; 
 }
 
+
 FFriendsAndChatStyle& FFriendsAndChatStyle::SetOfflineBrush(const FSlateBrush& InOffLine)
 {
 	OfflineBrush = InOffLine; 
 	return *this; 
 }
 
+
 FFriendsAndChatStyle& FFriendsAndChatStyle::SetOnlineBrush(const FSlateBrush& InOnLine)
 {
 	OnlineBrush = InOnLine; 
 	return *this; 
 }
 
+
 FFriendsAndChatStyle& FFriendsAndChatStyle::SetBackgroundBrush(const FSlateBrush& InBackground)
 {
 	Background = InBackground;
 	return *this;
 }
 
+
 FFriendsAndChatStyle& FFriendsAndChatStyle::SetTextStyle(const FTextBlockStyle& InTextStle)
 {
 	TextStyle = InTextStle;
 	return *this;
 }
 
+
 FFriendsAndChatStyle& FFriendsAndChatStyle::SetFontStyle(const FSlateFontInfo& InFontStyle)
 {
 	FriendsFontStyle = InFontStyle;
 	return *this;
 }
 
+
 FFriendsAndChatStyle& FFriendsAndChatStyle::SetCheckboxToggleStyle(const FCheckBoxStyle& InCheckboxStyle)
 {
 	CheckBoxToggleStyle = InCheckboxStyle;
 	return *this;
 }
 
+
 FFriendsAndChatStyle& FFriendsAndChatStyle::SetMenuSetColor(const FLinearColor& InSetColor)
 {
 	MenuSetColor = InSetColor;
 	return *this;
 }
 
+
 FFriendsAndChatStyle& FFriendsAndChatStyle::SetMenuUnsetColor(const FLinearColor& InUnsetColor)
 {
 	MenuUnsetColor = InUnsetColor;
 	return *this;
 }
 
+
 FFriendsAndChatStyle& FFriendsAndChatStyle::SetFriendItemStyle(const FComboButtonStyle& InFriendItemStyle)
 {
 	FriendItemStyle = InFriendItemStyle;
 	return *this;
 }
 
+
 const FFriendsAndChatStyle& FFriendsAndChatStyle::GetDefault()
 {
 	static FFriendsAndChatStyle Default;
diff --git a/Engine/Source/Developer/FriendsAndChat/Private/FriendsMessageManager.cpp b/Engine/Source/Developer/FriendsAndChat/Private/FriendsMessageManager.cpp
--- a/Engine/Source/Developer/FriendsAndChat/Private/FriendsMessageManager.cpp
+++ b/Engine/Source/Developer/FriendsAndChat/Private/FriendsMessageManager.cpp
@@ -1,13 +1,11 @@
 // Copyright 1998-2014 Epic Games, Inc. All Rights Reserved.
 
-/*=============================================================================
-	FriendsMessageManager.cpp: Implements the FFriendsMessageManager class.
-=============================================================================*/
-
 #include "FriendsAndChatPrivatePCH.h"
 
+
 #define LOCTEXT_NAMESPACE "FriendsMessageManager"
 
+
 /* FFriendsMessageManager structors
  *****************************************************************************/
 
@@ -18,8 +16,10 @@ FFriendsMessageManager::FFriendsMessageManager( )
 	PingMcpCountdown = 0;
 }
 
+
 FFriendsMessageManager::~FFriendsMessageManager( )
-{}
+{ }
+
 
 /* IFriendsMessageManager interface
  *****************************************************************************/
@@ -29,6 +29,7 @@ void FFriendsMessageManager::StartupManager()
 	FParse::Value( FCommandLine::Get(), TEXT( "LaunchGameMessage=" ), GameLaunchMessageID );
 }
 
+
 void FFriendsMessageManager::Init( FOnFriendsNotification& NotificationDelegate, bool bInCanJointGame )
 {
 	// Clear existing data
@@ -39,7 +40,7 @@ void FFriendsMessageManager::Init( FOnFriendsNotification& NotificationDelegate,
 
 	OnlineSubMcp = static_cast< FOnlineSubsystemMcp* >( IOnlineSubsystem::Get( TEXT( "MCP" ) ) );
 
-	if (OnlineSubMcp != NULL &&
+	if (OnlineSubMcp != nullptr &&
 		OnlineSubMcp->GetIdentityInterface().IsValid() &&
 		OnlineSubMcp->GetMessageInterface().IsValid())
 	{
@@ -73,11 +74,13 @@ void FFriendsMessageManager::Init( FOnFriendsNotification& NotificationDelegate,
 	}
 }
 
+
 void FFriendsMessageManager::SetMessagePolling( bool bStart )
 {
 	bPollForMessages = bStart;
 }
 
+
 bool FFriendsMessageManager::Tick( float Delta )
 {
 	PingMcpCountdown -= Delta;
@@ -118,21 +121,25 @@ bool FFriendsMessageManager::Tick( float Delta )
 	return true;
 }
 
+
 void FFriendsMessageManager::RequestEnumerateMessages()
 {
 	OnlineSubMcp->GetMessageInterface()->EnumerateMessages( 0 );
 }
 
+
 void FFriendsMessageManager::InviteFriendToGame( TSharedRef<FUniqueNetId> FriendID )
 {
 	GameInvitesToSend.AddUnique( FriendID );
 }
 
+
 void FFriendsMessageManager::RequestJoinAGame( TSharedRef<FUniqueNetId> FriendID )
 {
 	GameJoingRequestsToSend.AddUnique( FriendID );
 }
 
+
 void FFriendsMessageManager::SendGameInviteRequest()
 {
 	TSharedPtr<FUniqueNetId> UserId = OnlineSubMcp->GetIdentityInterface()->GetUniquePlayerId(0);
@@ -149,6 +156,7 @@ void FFriendsMessageManager::SendGameInviteRequest()
 	GameInvitesToSend.RemoveAt( 0 );
 }
 
+
 void FFriendsMessageManager::SendGameJoinRequest()
 {
 	TSharedPtr<FUniqueNetId> UserId = OnlineSubMcp->GetIdentityInterface()->GetUniquePlayerId( 0 );
@@ -166,6 +174,7 @@ void FFriendsMessageManager::SendGameJoinRequest()
 	GameJoingRequestsToSend.RemoveAt( 0 );
 }
 
+
 void FFriendsMessageManager::SendChatMessageRequest()
 {
 	TSharedPtr<FUniqueNetId> UserId = OnlineSubMcp->GetIdentityInterface()->GetUniquePlayerId(0);
@@ -188,36 +197,43 @@ void FFriendsMessageManager::SendChatMessageRequest()
 	ChatMessagesToSend.RemoveAt( 0 );
 }
 
+
 void FFriendsMessageManager::SendMessage( TSharedPtr< FFriendStuct > FriendID, const FText& Message )
 {
 	ChatMessagesToSend.Add( FOutGoingChatMessage( FriendID->GetOnlineUser()->GetUserId(), Message ) );
 }
 
+
 void FFriendsMessageManager::ClearGameInvites()
 {
 	bClearInvites = true;
 }
 
+
 void FFriendsMessageManager::SetUnhandledNotification( TSharedRef< FUniqueNetId > NetID )
 {
 	UnhandledNetID = NetID;
 }
 
+
 EFriendsMessageManagerState::Type FFriendsMessageManager::GetState()
 {
 	return ManagerState;
 }
 
+
 TArray< TSharedPtr< FFriendsAndChatMessage > > FFriendsMessageManager::GetChatMessages()
 {
 	return ChatMessages;
 }
 
+
 FOnMessagesUpdated& FFriendsMessageManager::OnChatListUpdated()
 {
 	return OChatListUpdatedDelegate;
 }
 
+
 void FFriendsMessageManager::SetState( EFriendsMessageManagerState::Type NewState )
 {
 	ManagerState = NewState;
@@ -272,10 +288,11 @@ void FFriendsMessageManager::SetState( EFriendsMessageManagerState::Type NewStat
 	}
 }
 
+
 void FFriendsMessageManager::Logout()
 {
 	FOnlineSubsystemMcp* OnlineSubMcp = static_cast< FOnlineSubsystemMcp* >( IOnlineSubsystem::Get( TEXT( "MCP" ) ) );
-	if ( OnlineSubMcp != NULL )
+	if ( OnlineSubMcp != nullptr )
 	{
 		OnlineSubMcp->GetMessageInterface()->ClearOnEnumerateMessagesCompleteDelegate(0, OnEnumerateMessagesCompleteDelegate);
 		OnlineSubMcp->GetMessageInterface()->ClearOnReadMessageCompleteDelegate(0, OnReadMessageCompleteDelegate);
@@ -285,6 +302,7 @@ void FFriendsMessageManager::Logout()
 	FTicker::GetCoreTicker().RemoveTicker( UpdateMessagesTickerDelegate );
 }
 
+
 void FFriendsMessageManager::OnEnumerateMessagesComplete(int32 LocalPlayer, bool bWasSuccessful, const FString& ErrorStr)
 {
 	if ( bWasSuccessful )
@@ -331,6 +349,7 @@ void FFriendsMessageManager::OnEnumerateMessagesComplete(int32 LocalPlayer, bool
 	SetState( EFriendsMessageManagerState::Idle );
 }
 
+
 void FFriendsMessageManager::OnReadMessageComplete(int32 LocalPlayer, bool bWasSuccessful, const FUniqueMessageId& MessageId, const FString& ErrorStr)
 {
 	if ( bWasSuccessful )
@@ -400,6 +419,7 @@ void FFriendsMessageManager::OnReadMessageComplete(int32 LocalPlayer, bool bWasS
 	SetState( EFriendsMessageManagerState::Idle );
 }
 
+
 void FFriendsMessageManager::OnSendMessageComplete(int32 LocalPlayer, bool bWasSuccessful, const FString& ErrorStr)
 {
 	if ( bWasSuccessful )
@@ -410,6 +430,7 @@ void FFriendsMessageManager::OnSendMessageComplete(int32 LocalPlayer, bool bWasS
 	SetState( EFriendsMessageManagerState::Idle );
 }
 
+
 void FFriendsMessageManager::OnDeleteMessageComplete(int32 LocalPlayer, bool bWasSuccessful, const FUniqueMessageId& MessageId, const FString& ErrorStr)
 {
 	// done with this part of the test if no more messages to delete
@@ -417,18 +438,21 @@ void FFriendsMessageManager::OnDeleteMessageComplete(int32 LocalPlayer, bool bWa
 	SetState( EFriendsMessageManagerState::Idle );
 }
 
+
 FReply FFriendsMessageManager::HandleMessageAccepted( TSharedPtr< FFriendsAndChatMessage > InNotificationMessage )
 {
 	NotficationMessages.Remove( InNotificationMessage );
 	return FReply::Handled();
 }
 
+
 FReply FFriendsMessageManager::HandleMessageDeclined( TSharedPtr< FFriendsAndChatMessage > InNotificationMessage )
 {
 	NotficationMessages.Remove( InNotificationMessage );
 	return FReply::Handled();
 }
 
+
 void FFriendsMessageManager::ResendMessage()
 {
 	if ( UnhandledNetID.IsValid() )
@@ -442,9 +466,11 @@ void FFriendsMessageManager::ResendMessage()
 	UnhandledNetID.Reset();
 }
 
+
 /* FFriendsMessageManager system singletons
 *****************************************************************************/
-TSharedPtr< FFriendsMessageManager > FFriendsMessageManager::SingletonInstance = NULL;
+
+TSharedPtr< FFriendsMessageManager > FFriendsMessageManager::SingletonInstance = nullptr;
 
 TSharedRef< FFriendsMessageManager > FFriendsMessageManager::Get()
 {
@@ -456,9 +482,11 @@ TSharedRef< FFriendsMessageManager > FFriendsMessageManager::Get()
 	return SingletonInstance.ToSharedRef();
 }
 
+
 void FFriendsMessageManager::Shutdown()
 {
 	SingletonInstance.Reset();
 }
 
+
 #undef LOCTEXT_NAMESPACE
