Changing the automation max slider from 10 to 1000.

diff --git a/Engine/Source/Developer/AutomationWindow/Private/SAutomationWindow.cpp b/Engine/Source/Developer/AutomationWindow/Private/SAutomationWindow.cpp
--- a/Engine/Source/Developer/AutomationWindow/Private/SAutomationWindow.cpp
+++ b/Engine/Source/Developer/AutomationWindow/Private/SAutomationWindow.cpp
@@ -1032,9 +1032,9 @@ TSharedRef< SWidget > SAutomationWindow::GenerateTestsOptionsMenuContent( )
 			[
 				SNew(SSpinBox<int32>)
 				.MinValue(1)
-				.MaxValue(10)
+				.MaxValue(1000)
 				.MinSliderValue(1)
-				.MaxSliderValue(10)
+				.MaxSliderValue(1000)
 				.Value(this,&SAutomationWindow::GetRepeatCount)
 				.OnValueChanged(this,&SAutomationWindow::OnChangeRepeatCount)
 				.IsEnabled( this, &SAutomationWindow::IsAutomationControllerIdle )
