SessionFrontend: Consistent warning messages when no session instances selected.

diff --git a/Engine/Source/Developer/AutomationWindow/Private/SAutomationWindow.cpp b/Engine/Source/Developer/AutomationWindow/Private/SAutomationWindow.cpp
--- a/Engine/Source/Developer/AutomationWindow/Private/SAutomationWindow.cpp
+++ b/Engine/Source/Developer/AutomationWindow/Private/SAutomationWindow.cpp
@@ -206,194 +206,208 @@ void SAutomationWindow::Construct( const FArguments& InArgs, const IAutomationCo
 	this->ChildSlot
 	[
 		SNew(SOverlay)
+
 		+ SOverlay::Slot()
-		[
-			SNew( SSplitter )
-			.Orientation(Orient_Vertical)
-			+SSplitter::Slot()
-			.Value(0.66f)
-			[
-				//automation test panel
-				SAssignNew( MenuBar, SVerticalBox )
-				//ACTIONS
-				+SVerticalBox::Slot()
-				.AutoHeight()
-				[
-					SNew( SHorizontalBox )
-					+SHorizontalBox::Slot()
-					.AutoWidth()
-					.HAlign(HAlign_Left)
-					.VAlign(VAlign_Center)
-					[
-						SAutomationWindow::MakeAutomationWindowToolBar( AutomationWindowActions.ToSharedRef(), SharedThis(this) )
-					]
-				]
-				+SVerticalBox::Slot()
-				.FillHeight(1.0f)
-				.Padding(0.0f, 4.0f, 0.0f, 0.0f)
-				[
-					SNew(SOverlay)
-					+SOverlay::Slot()
-					[
-						SNew(SBorder)
-						.BorderImage(this, &SAutomationWindow::GetTestBackgroundBorderImage)
-						.Padding(3.0f)
-						[
-							//the actual table full of tests
-							TestTable.ToSharedRef()
-						]
-					]
-					+SOverlay::Slot()
-					.HAlign(HAlign_Center)
-					.VAlign(VAlign_Center)
-					[
-						SNew(SThrobber)	
-							.Visibility(this, &SAutomationWindow::GetTestsUpdatingThrobberVisibility)
-					]
-				]
-			]
-			+SSplitter::Slot()
-			.Value(0.33f)
 			[
-				SNew(SOverlay)
-				+SOverlay::Slot()
-				[
-					SNew(SBox)
-					.Visibility(this, &SAutomationWindow::GetTestGraphVisibility)
-					[
-						//Graphical Results Panel
-						SNew( SVerticalBox )
-						+SVerticalBox::Slot()
-						.AutoHeight()
-						[
-							SNew(SHorizontalBox)
-							+SHorizontalBox::Slot()
-							.HAlign(HAlign_Left)
-							[
-								SNew(STextBlock)
-								.Text( LOCTEXT("AutomationTest_GraphicalResults", "Automation Test Graphical Results:"))
-							]
-
-							+SHorizontalBox::Slot()
-							.HAlign(HAlign_Right)
-							.AutoWidth()
-							[
-								SNew(STextBlock)
-								.Text( LOCTEXT("AutomationTest_Display", "Display:"))
-							]
+				SNew( SSplitter )
+					.IsEnabled(this, &SAutomationWindow::HandleMainContentIsEnabled)
+					.Orientation(Orient_Vertical)
 
-							+SHorizontalBox::Slot()
-							.HAlign(HAlign_Right)
-							.AutoWidth()
-							[
-								SNew(SCheckBox)
-								.Style(FCoreStyle::Get(), "RadioButton")
-								.IsChecked(this, &SAutomationWindow::HandleResultDisplayTypeIsChecked, EAutomationGrapicalDisplayType::DisplayName)
-								.OnCheckStateChanged(this, &SAutomationWindow::HandleResultDisplayTypeStateChanged, EAutomationGrapicalDisplayType::DisplayName)
-								[
-									SNew(STextBlock)
-									.Text( LOCTEXT("AutomationTest_GraphicalResultsDisplayName", "Name"))
-								]
-							]
+				+ SSplitter::Slot()
+					.Value(0.66f)
+					[
+						//automation test panel
+						SAssignNew( MenuBar, SVerticalBox )
 
-							+SHorizontalBox::Slot()
-							.HAlign(HAlign_Right)
-							.AutoWidth()
+						//ACTIONS
+						+ SVerticalBox::Slot()
+							.AutoHeight()
 							[
-								SNew(SCheckBox)
-								.Style(FCoreStyle::Get(), "RadioButton")
-								.IsChecked(this, &SAutomationWindow::HandleResultDisplayTypeIsChecked, EAutomationGrapicalDisplayType::DisplayTime)
-								.OnCheckStateChanged(this, &SAutomationWindow::HandleResultDisplayTypeStateChanged, EAutomationGrapicalDisplayType::DisplayTime)
-								[
-									SNew(STextBlock)
-									.Text( LOCTEXT("AutomationTest_GraphicalResultsDisplayTime", "Time"))
-								]
+								SNew( SHorizontalBox )
+
+								+ SHorizontalBox::Slot()
+									.AutoWidth()
+									.HAlign(HAlign_Left)
+									.VAlign(VAlign_Center)
+									[
+										SAutomationWindow::MakeAutomationWindowToolBar( AutomationWindowActions.ToSharedRef(), SharedThis(this) )
+									]
 							]
-						]
 
-						+SVerticalBox::Slot()
-						.FillHeight(1.0f)
-						[
-							SNew(SBorder)
+						+ SVerticalBox::Slot()
+							.FillHeight(1.0f)
+							.Padding(0.0f, 4.0f, 0.0f, 0.0f)
 							[
-								SNew(SScrollBox)
-								+ SScrollBox::Slot()
-								[
-									SAssignNew(GraphicalResultBox, SAutomationGraphicalResultBox, InAutomationController)
-								]
+								SNew(SOverlay)
+									+ SOverlay::Slot()
+									[
+										SNew(SBorder)
+											.BorderImage(this, &SAutomationWindow::GetTestBackgroundBorderImage)
+											.Padding(3.0f)
+											[
+												//the actual table full of tests
+												TestTable.ToSharedRef()
+											]
+									]
+
+								+ SOverlay::Slot()
+									.HAlign(HAlign_Center)
+									.VAlign(VAlign_Center)
+									[
+										SNew(SThrobber)	
+											.Visibility(this, &SAutomationWindow::GetTestsUpdatingThrobberVisibility)
+									]
 							]
-						]
 					]
-				]
 
-				+SOverlay::Slot()
-				[
-					SNew(SBox)
-					.Visibility(this, &SAutomationWindow::GetTestLogVisibility)
+				+ SSplitter::Slot()
+					.Value(0.33f)
 					[
-						//results panel
-						SNew( SVerticalBox )
-						+SVerticalBox::Slot()
-						.AutoHeight()
-						[
-							SNew( STextBlock )
-							.Text( LOCTEXT("AutomationTest_Results", "Automation Test Results:") )
-						]
+						SNew(SOverlay)
 
-						+SVerticalBox::Slot()
-						.FillHeight(1.0f)
-						.Padding(0.0f, 4.0f, 0.0f, 0.0f)
-						[
-							//list of results for the selected test
-							SNew(SBorder)
+						+ SOverlay::Slot()
 							[
-								SAssignNew(LogListView, SListView<TSharedPtr<FAutomationOutputMessage> >)
-								.ItemHeight(18)
-								.ListItemsSource(&LogMessages)
-								.SelectionMode(ESelectionMode::Multi)
-								.OnGenerateRow(this, &SAutomationWindow::OnGenerateWidgetForLog)
-								.OnSelectionChanged(this, &SAutomationWindow::HandleLogListSelectionChanged)
+								SNew(SBox)
+									.Visibility(this, &SAutomationWindow::GetTestGraphVisibility)
+									[
+										//Graphical Results Panel
+										SNew( SVerticalBox )
+
+										+ SVerticalBox::Slot()
+											.AutoHeight()
+											[
+												SNew(SHorizontalBox)
+
+												+ SHorizontalBox::Slot()
+													.HAlign(HAlign_Left)
+													[
+														SNew(STextBlock)
+															.Text( LOCTEXT("AutomationTest_GraphicalResults", "Automation Test Graphical Results:"))
+													]
+
+												+ SHorizontalBox::Slot()
+													.HAlign(HAlign_Right)
+													.AutoWidth()
+													[
+														SNew(STextBlock)
+															.Text( LOCTEXT("AutomationTest_Display", "Display:"))
+													]
+
+												+ SHorizontalBox::Slot()
+													.HAlign(HAlign_Right)
+													.AutoWidth()
+													[
+														SNew(SCheckBox)
+															.Style(FCoreStyle::Get(), "RadioButton")
+															.IsChecked(this, &SAutomationWindow::HandleResultDisplayTypeIsChecked, EAutomationGrapicalDisplayType::DisplayName)
+															.OnCheckStateChanged(this, &SAutomationWindow::HandleResultDisplayTypeStateChanged, EAutomationGrapicalDisplayType::DisplayName)
+															[
+																SNew(STextBlock)
+																	.Text( LOCTEXT("AutomationTest_GraphicalResultsDisplayName", "Name"))
+															]
+													]
+
+												+ SHorizontalBox::Slot()
+													.HAlign(HAlign_Right)
+													.AutoWidth()
+													[
+														SNew(SCheckBox)
+															.Style(FCoreStyle::Get(), "RadioButton")
+															.IsChecked(this, &SAutomationWindow::HandleResultDisplayTypeIsChecked, EAutomationGrapicalDisplayType::DisplayTime)
+															.OnCheckStateChanged(this, &SAutomationWindow::HandleResultDisplayTypeStateChanged, EAutomationGrapicalDisplayType::DisplayTime)
+															[
+																SNew(STextBlock)
+																	.Text( LOCTEXT("AutomationTest_GraphicalResultsDisplayTime", "Time"))
+															]
+													]
+											]
+
+										+ SVerticalBox::Slot()
+											.FillHeight(1.0f)
+											[
+												SNew(SBorder)
+												[
+													SNew(SScrollBox)
+
+													+ SScrollBox::Slot()
+														[
+															SAssignNew(GraphicalResultBox, SAutomationGraphicalResultBox, InAutomationController)
+														]
+												]
+											]
+									]
 							]
-						]
 
-						+SVerticalBox::Slot()
-						.AutoHeight()
-						.Padding(0.0f, 4.0f, 0.0f, 0.0f)
-						[
-							SNew(SBorder)
-							.BorderImage(FEditorStyle::GetBrush("ToolPanel.GroupBorder"))
-							.Padding(FMargin(8.0f, 6.0f))
+						+ SOverlay::Slot()
 							[
-								// Add the command bar
-								SAssignNew(CommandBar, SAutomationWindowCommandBar, NotificationList)
-								.OnCopyLogClicked(this, &SAutomationWindow::HandleCommandBarCopyLogClicked)
+								SNew(SBox)
+									.Visibility(this, &SAutomationWindow::GetTestLogVisibility)
+									[
+										//results panel
+										SNew( SVerticalBox )
+										+ SVerticalBox::Slot()
+											.AutoHeight()
+											[
+												SNew( STextBlock )
+													.Text( LOCTEXT("AutomationTest_Results", "Automation Test Results:") )
+											]
+
+										+ SVerticalBox::Slot()
+											.FillHeight(1.0f)
+											.Padding(0.0f, 4.0f, 0.0f, 0.0f)
+											[
+												//list of results for the selected test
+												SNew(SBorder)
+													[
+														SAssignNew(LogListView, SListView<TSharedPtr<FAutomationOutputMessage> >)
+															.ItemHeight(18)
+															.ListItemsSource(&LogMessages)
+															.SelectionMode(ESelectionMode::Multi)
+															.OnGenerateRow(this, &SAutomationWindow::OnGenerateWidgetForLog)
+															.OnSelectionChanged(this, &SAutomationWindow::HandleLogListSelectionChanged)
+													]
+											]
+
+										+ SVerticalBox::Slot()
+											.AutoHeight()
+											.Padding(0.0f, 4.0f, 0.0f, 0.0f)
+											[
+												SNew(SBorder)
+													.BorderImage(FEditorStyle::GetBrush("ToolPanel.GroupBorder"))
+													.Padding(FMargin(8.0f, 6.0f))
+													[
+														// Add the command bar
+														SAssignNew(CommandBar, SAutomationWindowCommandBar, NotificationList)
+															.OnCopyLogClicked(this, &SAutomationWindow::HandleCommandBarCopyLogClicked)
+													]
+											]
+									]
 							]
-						]
 					]
-				]
 			]
-		]
+
 		+ SOverlay::Slot()
-		.HAlign( HAlign_Center )
-		.VAlign( VAlign_Center )
-		.Padding( 15.0f )
-		[
-			NotificationList
-		]
+			.HAlign( HAlign_Center )
+			.VAlign( VAlign_Center )
+			.Padding( 15.0f )
+			[
+				NotificationList
+			]
+
 		+ SOverlay::Slot()
-		.HAlign( HAlign_Center )
-		.VAlign( VAlign_Center )
-		.Padding( 15.0f )
-		[
-			SNew(SBorder)
-				.BorderImage(FEditorStyle::GetBrush("NotificationList.ItemBackground"))
-				.Padding(8.0f)
-				.Visibility(this, &SAutomationWindow::HandleSelectSessionOverlayVisibility)
-				[
-					SNew(STextBlock)
-						.Text(LOCTEXT("SelectSessionOverlayText", "Please select a session from the Session Browser"))
-				]
-		]
+			.HAlign( HAlign_Center )
+			.VAlign( VAlign_Center )
+			.Padding( 15.0f )
+			[
+				SNew(SBorder)
+					.BorderImage(FEditorStyle::GetBrush("NotificationList.ItemBackground"))
+					.Padding(8.0f)
+					.Visibility(this, &SAutomationWindow::HandleSelectSessionOverlayVisibility)
+					[
+						SNew(STextBlock)
+							.Text(LOCTEXT("SelectSessionOverlayText", "Please select at least one instance from the Session Browser"))
+					]
+			]
 	];
 
 	SessionManager->OnCanSelectSession().AddSP( this, &SAutomationWindow::HandleSessionManagerCanSelectSession );
@@ -1793,7 +1807,7 @@ FText SAutomationWindow::HandleAutomationHighlightText( ) const
 
 EVisibility SAutomationWindow::HandleSelectSessionOverlayVisibility( ) const
 {
-	if (SessionManager->GetSelectedSession().IsValid())
+	if (SessionManager->GetSelectedInstances().Num() > 0)
 	{
 		return EVisibility::Hidden;
 	}
@@ -2004,4 +2018,10 @@ bool SAutomationWindow::HandleItemCheckBoxIsEnabled( ) const
 }
 
 
+bool SAutomationWindow::HandleMainContentIsEnabled() const
+{
+	return (SessionManager->GetSelectedInstances().Num() > 0);
+}
+
+
 #undef LOCTEXT_NAMESPACE
diff --git a/Engine/Source/Developer/AutomationWindow/Private/SAutomationWindow.h b/Engine/Source/Developer/AutomationWindow/Private/SAutomationWindow.h
--- a/Engine/Source/Developer/AutomationWindow/Private/SAutomationWindow.h
+++ b/Engine/Source/Developer/AutomationWindow/Private/SAutomationWindow.h
@@ -138,6 +138,9 @@ private:
 	/** Callback for getting the enabled state of a test item. */
 	bool HandleItemCheckBoxIsEnabled( ) const;
 
+	/** Callback for getting the enabled state of the main content overlay. */
+	bool HandleMainContentIsEnabled() const;
+
 	/** Create the UI commands for the toolbar */
 	void CreateCommands();
 
diff --git a/Engine/Source/Developer/Profiler/Private/Widgets/SProfilerWindow.cpp b/Engine/Source/Developer/Profiler/Private/Widgets/SProfilerWindow.cpp
--- a/Engine/Source/Developer/Profiler/Private/Widgets/SProfilerWindow.cpp
+++ b/Engine/Source/Developer/Profiler/Private/Widgets/SProfilerWindow.cpp
@@ -1,9 +1,5 @@
 ﻿// Copyright 1998-2015 Epic Games, Inc. All Rights Reserved.
 
-/*=============================================================================
-	SProfilerWindow.cpp: Implements the SProfilerWindow class.
-=============================================================================*/
-
 #include "ProfilerPrivatePCH.h"
 #include "SProfilerSettings.h"
 #include "DesktopPlatformModule.h"
@@ -14,8 +10,10 @@
 #endif // WITH_EDITOR
 #include "SNotificationList.h"
 
+
 #define LOCTEXT_NAMESPACE "SProfilerWindow"
 
+
 static FText GetTextForNotification( const EProfilerNotificationTypes::Type NotificatonType, const ELoadingProgressStates::Type ProgressState, const FString& Filename, const float ProgressPercent = 0.0f )
 {
 	FText Result;
@@ -135,18 +133,18 @@ void SProfilerWindow::Construct( const FArguments& InArgs )
 						[
 							SNew(SBox)
 								.HeightOverride(48.0f)
+								.IsEnabled(this, &SProfilerWindow::IsProfilerEnabled)
 								[
 									SNew(SHorizontalBox)
-										.IsEnabled(this, &SProfilerWindow::IsProfilerEnabled)
 
-										+ SHorizontalBox::Slot()
-											.FillWidth(1.0f)
-											.Padding(0.0f)
-											.HAlign(HAlign_Fill)
-											.VAlign(VAlign_Fill)
-											[
-												SAssignNew(ProfilerMiniView, SProfilerMiniView)
-											]
+									+ SHorizontalBox::Slot()
+										.FillWidth(1.0f)
+										.Padding(0.0f)
+										.HAlign(HAlign_Fill)
+										.VAlign(VAlign_Fill)
+										[
+											SAssignNew(ProfilerMiniView, SProfilerMiniView)
+										]
 								]
 						]
 
@@ -471,6 +469,7 @@ EActiveTimerReturnType SProfilerWindow::UpdateActiveDuration( double InCurrentTi
 	return EActiveTimerReturnType::Continue;
 }
 
+
 void SProfilerWindow::OnMouseEnter(const FGeometry& MyGeometry, const FPointerEvent& MouseEvent)
 {
 	SCompoundWidget::OnMouseEnter(MyGeometry, MouseEvent);
