Copying //UE4/Dev-Editor to //UE4/Dev-Main (Source: //UE4/Dev-Editor @ 2973866)

==========================
MAJOR FEATURES + CHANGES
==========================

Change 2937390 on 2016/04/07 by [REDACTED_PERSON]

	Fixed slider to properly bubble unhandled OnKeyDown events

Change 2939672 on 2016/04/11 by [REDACTED_PERSON]

	Made a change to how file check out notifications work. Now the dirty package state is processed at the end of every tick, meaning that packages which are dirtied and then cleaned again are not processed. This fixes an issue where a number of child blueprints were flagged as needing checkout when a parent blueprint was compiled.  This also allows multiple packages which are dirtied at the same time to be treated as one transaction.

Change 2939686 on 2016/04/11 by [REDACTED_PERSON]

	A number of further improvements to mesh vertex color painting:
	* Lower LODs are now automatically fixed up for instances which were created in a previous bugged version of the engine.
	* Since lower LODs cannot currently have their vertex colors edited, their vertex colors are always derived from LOD0.
	* Fixed a bug when building lower LODs so that vertices in neighboring octree nodes are considered when looking for the nearest vertex from LOD0 which corresponds.
	* Fixed issue where static meshes with imported LODs would not have the lower LODs' override colors set when "Copy instance vertex colors to source mesh" was used (static meshes with generated LODs were always getting correct override colors).

Change 2939906 on 2016/04/11 by [REDACTED_PERSON]

	Automation - Adding several enhancements to the automation framework and improving the UI.
	* Tests in the UI now have a link to the source and line where they orginate.
	* There's now a general purpose latent lambda command you can use to run arbitrary code latently.
	* Added Inlined AddCommand for regular and networked commands to the base automation class, to avoid the use of the macro, which prevents breakpoints from working in lambda code.
	* Front end now has better column displays offering more room to the test name
	* Changed several events to the automation controller to multicast delegates so that many could hook them.
	* The UI now refreshes the selection after tests finish so that the output log updates.

Change 2939908 on 2016/04/11 by [REDACTED_PERSON]

	Automation - The editor import/export tests are now a complex test and actually sperate out all the tests that can be run, some trickiness was required on the filenames so that they didn't expand into more child tests in the UI. (replacing .'s with _'s)

Change 2940028 on 2016/04/11 by [REDACTED_PERSON]

	Automation - Removing the search box from the toolbar.  It's now inlined above the test tree.  Tweaking the padding to make it look more other windows and make everything not look so squished.  Recursive expansion now works on tests.

Change 2940066 on 2016/04/11 by [REDACTED_PERSON]

	Automation - Moving the filter group dropdown out of the toolbar and onto the line with the search box above the treeview - additional tweaks to it.

Change 2940092 on 2016/04/11 by [REDACTED_PERSON]

	PR #2248: Datatable select next row (Contributed by [REDACTED_PERSON])

Change 2940093 on 2016/04/11 by [REDACTED_PERSON]

	PR #2248: Datatable select next row (Contributed by [REDACTED_PERSON])

Change 2940157 on 2016/04/11 by [REDACTED_PERSON]

	Fixing FTextTest due to some changes made to how currency is formatted

Change 2940694 on 2016/04/12 by [REDACTED_PERSON]

	Fixed issue where vertex override colors were not being propagated correctly for generated lower LODs.

Change 2942379 on 2016/04/13 by [REDACTED_PERSON]

	Fixed issue where entering PIE while selecting an actor in Mesh Paint mode could lead to a MeshPaintStaticMeshAdapter holding onto an invalid pointer to an old mesh component, and causing a crash upon leaving the mode.  This can happen because, when loading a new streaming level, the proxy actor can be selected when starting PIE, which will subsequently be added to the tool's internal lists.  This needs to be added as a GC reference so that it can be NULLed when forcibly destroyed.

Change 2942947 on 2016/04/13 by [REDACTED_PERSON]

	Fixed crash when pasting a material function call node from one project to another in which it is not defined.

Change 2943452 on 2016/04/14 by [REDACTED_PERSON]

	Updated F4 debug key binding to match what's in ShowFlags.cpp
	PR #2197 (contributed by [REDACTED_PERSON])

Change 2943824 on 2016/04/14 by [REDACTED_PERSON]

	Make sure we cannot open the color picker when a property is edit const

Change 2943841 on 2016/04/14 by [REDACTED_PERSON]

	tooltip was add for every hierarchy import option

Change 2943927 on 2016/04/14 by [REDACTED_PERSON]

	Add Obj support for scene importer

	Github PR #2272

Change 2943967 on 2016/04/14 by [REDACTED_PERSON]

	Added relevant fields from FBodyInstance to the FoliageType customizations.

Change 2948397 on 2016/04/19 by [REDACTED_PERSON]

	Moved FSlateIcon definition to SlateCore

	It was previously declared as SLATE_API, despite its header residing inside SlateCore. Reviewed by [REDACTED_PERSON].

Change 2948805 on 2016/04/19 by [REDACTED_PERSON]

	Editor: Deprecated FName UEdGraphNode::GetPaletteIcon(FLinearColor&); in favor of FSlateIcon UEdGraphNode::GetIconAndTint(FLinearColor&); to allow for icons in external style sets to be used.

	  - Previously, all icons were assumed to reside within FEditorStyle, which is not the case and would create broken icons in the graph editor. All relevant code has been updated to use FSlateIcon structures instead of a simple name.
	  - This change required a significant overhaul to FClassIconFinder to support FSlateIcons. To keep the API clean, FSlateIconFinder now deals with FSlateIcon class icon finding operations, and FClassIconFinder for the most part just adds actor specific logic.

Change 2950658 on 2016/04/20 by [REDACTED_PERSON]

	Skinxx workflow, we now output an error if there is mix of material with skinxx and some with no skinxx suffix

Change 2950663 on 2016/04/20 by [REDACTED_PERSON]

	When exporting to fbx we have to export each material instance as one fbx material

Change 2951240 on 2016/04/21 by [REDACTED_PERSON]

	Make sure light are render properly after importing a fbx scene

Change 2951421 on 2016/04/21 by [REDACTED_PERSON]

	fbx skeletalmesh import now support mesh hierarchy

Change 2955873 on 2016/04/26 by [REDACTED_PERSON]

	PR #2225: Fix working package directory from the launch profiles (Contributed by [REDACTED_PERSON])

Change 2955965 on 2016/04/26 by [REDACTED_PERSON]

	Merging //UE4/Dev-Main to Dev-Editor (//UE4/Dev-Editor)

Change 2956717 on 2016/04/26 by [REDACTED_PERSON]

	Editor: World Outliner now correctly calls ProcessEditDelete on editor modes that have asked to process delete operations

Change 2956822 on 2016/04/26 by [REDACTED_PERSON]

	Editor: Fixed actors not being removed from the scene outliner when they are added and removed on the same frame

Change 2956931 on 2016/04/26 by [REDACTED_PERSON]

	New Module - UATHelper - Moving the UAT launching code from the MainFrame module into a reusable module other modules can trigger.

Change 2956932 on 2016/04/26 by [REDACTED_PERSON]

	Plugins - Now allowing you to package a plugin from the plugin browsing view.  Still work in progress.

Change 2957164 on 2016/04/26 by [REDACTED_PERSON]

	Hot Reload - Fixing hot reload, it no longer creates a temporary copy of the module manager.  Making the copy constructor private on the module manager to prevent this in the future.

Change 2957165 on 2016/04/26 by [REDACTED_PERSON]

	Fixing the Editor Mode plugin sample, it no longer provides a bad starting example for where to create your widgets.

Change 2957510 on 2016/04/27 by [REDACTED_PERSON]

	PR #2198: Git Plugin implement the Sync operation to update local files using the git pull --rebase command (Contributed by [REDACTED_PERSON])

Change 2957511 on 2016/04/27 by [REDACTED_PERSON]

	Editor: Make favorites button on details panel non-focusable
	  - This was preventing users being able to tab between value fields on the details panel

Change 2957610 on 2016/04/27 by [REDACTED_PERSON]

	PR #1836: Git plugin: make initial commit when initializing new project (Contributed by [REDACTED_PERSON])

Change 2957667 on 2016/04/27 by [REDACTED_PERSON]

	Fixed crash that could happen in FTextLayout::GetLineViewIndexForTextLocation if passed a bad location

Change 2958035 on 2016/04/27 by [REDACTED_PERSON]

	Fixing the DesignerRebuild flag detection so that we can just refresh the slate widget without recreating the preview UObject, which causes the destruction of the details panel, and the slate widget recreation was the only part that was required.

Change 2958272 on 2016/04/27 by [REDACTED_PERSON]

	Added FAssetData::GetTagValue to handle getting asset tag values in a type-correct way

	This allows type-conversion using LexicalConversion, and also has specializations for FString, FText, and FName.

Change 2958348 on 2016/04/27 by [REDACTED_PERSON]

	PR #2282: Slate font shutdown order fix (Contributed by [REDACTED_PERSON])

Change 2958352 on 2016/04/27 by [REDACTED_PERSON]

	Fixed the subtitle manager updating the wrong list of subtitles

Change 2958390 on 2016/04/27 by [REDACTED_PERSON]

	Removed some old placement-new style array insertions

Change 2959360 on 2016/04/28 by [REDACTED_PERSON]

	Fixed potential crash when mesh painting actors whose geometry adapters are no longer registered.

Change 2959724 on 2016/04/28 by [REDACTED_PERSON]

	Merging hardware survey gating logic from 4.10

Change 2959807 on 2016/04/28 by [REDACTED_PERSON]

	Removed deprecated function call

Change 2959894 on 2016/04/28 by [REDACTED_PERSON]

	Fix for scroll offset being clamped by content size, not scroll max

Change 2960048 on 2016/04/28 by [REDACTED_PERSON]

	Added FAssetData::GetTagValueRef to go along with FAssetData::GetTagValue

Change 2960782 on 2016/04/29 by [REDACTED_PERSON]

	Updating code to use the new FText aware asset registry tag functions

Change 2960885 on 2016/04/29 by [REDACTED_PERSON]

	Updating code to use the new FText aware asset registry tag functions

Change 2961170 on 2016/04/29 by [REDACTED_PERSON]

	Updating code to use the new FText aware asset registry tag functions

Change 2961171 on 2016/04/29 by [REDACTED_PERSON]

	Updating code to use the new FText aware asset registry tag functions

Change 2961173 on 2016/04/29 by [REDACTED_PERSON]

	Removed some inline duplication on the specialized template functions

Change 2963124 on 2016/05/02 by [REDACTED_PERSON]

	FExternalDragOperation can now contain both text and file data at the same time

	This better mirrors what the OS level drag-and-drop operations are capable of, and some applications will actually give you both bits of data at the same time.

Change 2963175 on 2016/05/02 by [REDACTED_PERSON]

	Updated some font editor tooltips to be more descriptive

Change 2963290 on 2016/05/02 by [REDACTED_PERSON]

	The Localise UAT command can now be run with a null localisation provider

Change 2963305 on 2016/05/02 by [REDACTED_PERSON]

	Fixed minor typo

Change 2963402 on 2016/05/02 by [REDACTED_PERSON]

	Cleaned up all the current localization key conflicts and warnings from gathering Engine code

Change 2963415 on 2016/05/02 by [REDACTED_PERSON]

	Rephrased a message that could generate a CIS warning

Change 2964184 on 2016/05/03 by [REDACTED_PERSON]

	Fixed duplicate "Font" entry in asset picker menu

	This was caused by PropertyCustomizationHelpers::GetNewAssetFactoriesForClasses using CanCreateNew rather than ShouldShowInNewMenu, as UFont has two factories, but one is supposed to be hidden from the UI.

	We also now make sure the factories are sorted by display name before being shown in the UI.

Change 2966108 on 2016/05/04 by [REDACTED_PERSON]

	Engine - Rearranging the order of ELoadingPhase's enums so that they match the loading order of modules.

Change 2966113 on 2016/05/04 by [REDACTED_PERSON]

	[Engine Loop Change] UEngine now defines a Start() function, that subclasses can use to start game related things after initialization of the engine.  This is done so that after the Init() call on UEngine, we can then perform a module load for the ELoadingPhase::PostEngineInit phase of loading, then inform the UEngine that it's time to start the game.  Therefore, UGameEngine now tells the GameInstance to Start during this phase now.

Change 2966121 on 2016/05/04 by [REDACTED_PERSON]

	Config writing improvements when dealing with property values

	This updates FConfigFile::ShouldExportQuotedString to make sure that a property value containing any characters that FParse::LineExtended will consume when parsing back in the config file (such as { and }, or a trailing \) cause the string to be quoted.

	This also adds FConfigFile::GenerateExportedPropertyLine to generate the INI key->value lines in a consistent and correctly escaped way, and makes sure that everything that writes out lines to a config file uses it.

	FConfigCacheIni::SetString and FConfigCacheIni::SetText have been updated to update the value even if it only differs by case.

	UObject::SaveConfig and UObject::LoadConfig have had some code whitespace fix-up (from a bad merge).

Change 2966122 on 2016/05/04 by [REDACTED_PERSON]

	Added a setting to control dialogue wave audio filenames

Change 2966481 on 2016/05/04 by [REDACTED_PERSON]

	PR #2336: BUGFIX: Selection of objects in the Content browser from WorldSettings (Contributed by [REDACTED_PERSON])

Change 2966887 on 2016/05/04 by [REDACTED_PERSON]

	PR #2336: BUGFIX: Selection of objects in the Content browser from WorldSettings (Contributed by [REDACTED_PERSON])

Change 2967488 on 2016/05/05 by [REDACTED_PERSON]

	Changes to support packaging plugins from the editor.

	* UBT now has an option to explicitly disable hot-reloading in any circumstances.
	* When running with -module arguments for a monolithic target, UBT will no longer try to relink the executable in source builds (so it's possible to compile plugin libs outside of an installed engine build without having already built UE4Game).
	* When packaging, a temporary host project is always generated in the output directory to avoid invalidating intermediates in the source directory.
	* An empty Config\FilterPlugin.ini file is written out with instructions on how to list additional files to package if it is not already present.

Change 2967947 on 2016/05/05 by [REDACTED_PERSON]

	PR #2358: Properly display Mip Level Count and Format for UTexture2DDynamic Textures (Contributed by Allegorithmic)

Change 2968333 on 2016/05/05 by [REDACTED_PERSON]

	Fixed MultiLine not working with arrays of string or text properties

	- The detail customizations for FString and FText properties now read the meta-data off the correct property.
	- The UDS editor now lets you set the "MultiLine" meta-data on arrays of FString and FText properties.
	- Fixed changing the "MultiLine" flag on a UDS property not rebuilding the default value editor.
	- Fixed the default values panel in the UDS editor having a title area.

Change 2968999 on 2016/05/06 by [REDACTED_PERSON]

	Fixed infinite loop in the editor if a directory that is being watched is deleted

Change 2969105 on 2016/05/06 by [REDACTED_PERSON]

	Fixed issue where opening a submenu while the parent menu had a text box focused would lead to a crash. The graph node comment text widget now only dismisses all menus if the text commit info implies that it was committed by some user action.

Change 2969440 on 2016/05/06 by [REDACTED_PERSON]

	Significant performance improvements when pasting a large amount of text

Change 2969619 on 2016/05/06 by [REDACTED_PERSON]

	Auto-reimport is now disabled inside an editor running in unattended mode

Change 2969621 on 2016/05/06 by [REDACTED_PERSON]

	Added the ability to override the subtitle used on a dialogue wave

	This is useful for effort sounds, plus some other cases, such as characters speaking in a foreign language not known to the player.

Change 2970588 on 2016/05/09 by [REDACTED_PERSON]

	Fix typo in operator expression in UEndUserSettings::SetSendAnonymousUsageDataToEpic()
	[UE-26958] - GitHub 2056 : Fixing typo in the operator
	#2056

Change 2971151 on 2016/05/09 by [REDACTED_PERSON]

	Logging ensure fails as errors. Automated tests with ensure fails will be unsuccessful.
	[UE-19579] - If an ensure() fails within an automated test, the test can still show a positive result.
	[UE-26575] - GitHub 2030 : Add error-severity message to log on ensure.

	PR #2030

Change 2971267 on 2016/05/09 by [REDACTED_PERSON]

	Wrong parameter when calling GetImportOptions

Change 2972073 on 2016/05/10 by [REDACTED_PERSON]

	Fixed UModel methods which make surfaces as modified.

Change 2972329 on 2016/05/10 by [REDACTED_PERSON]

	Merging //UE4/Dev-Main to Dev-Editor (//UE4/Dev-Editor)

Change 2972887 on 2016/05/10 by [REDACTED_PERSON]

	We now import the geometric transform also when we uncheck the absolute transform in the vertex.

Change 2973664 on 2016/05/11 by [REDACTED_PERSON]

	Merging //UE4/Dev-Main to Dev-Editor (//UE4/Dev-Editor)

Change 2973717 on 2016/05/11 by [REDACTED_PERSON]

	Fixing compiler issues from main merge.

Change 2973846 on 2016/05/11 by [REDACTED_PERSON]

	Exposed FConfigValue::ExpandValue and added FConfigValue::CollapseValue

	These are both static and can be used to expand or collapse the macros used in our config files (mostly when dealing with paths), in code that has to deal with the config system, but isn't internal to the config system (mostly things that deal with default configs outside of UObjects).

	The old non-static version of FConfigValue::ExpandValue is now FConfigValue::ExpandValueInternal, which just calls FConfigValue::ExpandValue on SavedValue and ExpandedValue.

	This also changes some code that was using FString.Replace to use FString.ReplaceInline. This reduces allocations, and also allows us to avoid another string comparison to see whether the strings are identical (as ReplaceInline returns the number of replacements that were made).

Change 2973847 on 2016/05/11 by [REDACTED_PERSON]

	Changing the loading phase in the localization dashboard now writes to the default config

Change 2973866 on 2016/05/11 by [REDACTED_PERSON]

	Deprecated some functions that were taking an unused position.

	These unused parameters caused confusion and lead to UE-30276. The old versions have been deprecated, and new versions without those parameters have been added. Existing code has been updated to call the non-deprecated version.

	- FViewportFrame::ResizeFrame
	- FSceneViewport::ResizeFrame
	- FSceneViewport::ResizeViewport

diff --git a/Engine/Config/BaseEngine.ini b/Engine/Config/BaseEngine.ini
--- a/Engine/Config/BaseEngine.ini
+++ b/Engine/Config/BaseEngine.ini
@@ -1586,6 +1586,7 @@ UnfocusedVolumeMultiplier=0.0
 DefaultSoundClassName=/Engine/EngineSounds/Master.Master
 LowPassFilterResonance=0.9
 MaximumConcurrentStreams=2
+DialogueFilenameFormat="{DialogueGuid}_{ContextId}"
 
 [/Script/Engine.SoundGroups]
 +SoundGroupProfiles=(SoundGroup=SOUNDGROUP_Default, bAlwaysDecompressOnLoad=false, DecompressedDuration=5)
diff --git a/Engine/Config/BaseInput.ini b/Engine/Config/BaseInput.ini
--- a/Engine/Config/BaseInput.ini
+++ b/Engine/Config/BaseInput.ini
@@ -31,7 +31,7 @@ bRequireCtrlToNavigateAutoComplete=False
 +DebugExecBindings=(Key=F1,Command="viewmode wireframe", bIgnoreShift=True)
 +DebugExecBindings=(Key=F2,Command="viewmode unlit")
 +DebugExecBindings=(Key=F3,Command="viewmode lit")
-+DebugExecBindings=(Key=F4,Command="viewmode detaillighting")
++DebugExecBindings=(Key=F4,Command="viewmode lit_detaillighting")
 +DebugExecBindings=(Key=F5,Command="viewmode shadercomplexity")
 +DebugExecBindings=(Key=F9,Command="shot showui")
 +DebugExecBindings=(Key=Period,Command="RECOMPILESHADERS CHANGED",Control=True,Shift=True)
diff --git a/Engine/Plugins/2D/Paper2D/Source/Paper2DEditor/Private/Atlasing/AtlasAssetTypeActions.cpp b/Engine/Plugins/2D/Paper2D/Source/Paper2DEditor/Private/Atlasing/AtlasAssetTypeActions.cpp
--- a/Engine/Plugins/2D/Paper2D/Source/Paper2DEditor/Private/Atlasing/AtlasAssetTypeActions.cpp
+++ b/Engine/Plugins/2D/Paper2D/Source/Paper2DEditor/Private/Atlasing/AtlasAssetTypeActions.cpp
@@ -44,15 +44,13 @@ uint32 FAtlasAssetTypeActions::GetCategories()
 
 FText FAtlasAssetTypeActions::GetAssetDescription(const FAssetData& AssetData) const
 {
-	if (const FString* pDescription = AssetData.TagsAndValues.Find(GET_MEMBER_NAME_CHECKED(UPaperSpriteAtlas, AtlasDescription)))
+	FString Description = AssetData.GetTagValueRef<FString>(GET_MEMBER_NAME_CHECKED(UPaperSpriteAtlas, AtlasDescription));
+	if (!Description.IsEmpty())
 	{
-		if (!pDescription->IsEmpty())
-		{
-			const FString DescriptionStr(*pDescription);
-			return FText::FromString(DescriptionStr.Replace(TEXT("\\n"), TEXT("\n")));
-		}
+		Description.ReplaceInline(TEXT("\\n"), TEXT("\n"));
+		return FText::FromString(MoveTemp(Description));
 	}
-
+	
 	return FText::GetEmpty();
 }
 
diff --git a/Engine/Plugins/2D/Paper2D/Source/Paper2DEditor/Private/FlipbookEditor/FlipbookEditorCommands.cpp b/Engine/Plugins/2D/Paper2D/Source/Paper2DEditor/Private/FlipbookEditor/FlipbookEditorCommands.cpp
--- a/Engine/Plugins/2D/Paper2D/Source/Paper2DEditor/Private/FlipbookEditor/FlipbookEditorCommands.cpp
+++ b/Engine/Plugins/2D/Paper2D/Source/Paper2DEditor/Private/FlipbookEditor/FlipbookEditorCommands.cpp
@@ -6,7 +6,7 @@
 //////////////////////////////////////////////////////////////////////////
 // FFlipbookEditorCommands
 
-#define LOCTEXT_NAMESPACE ""
+#define LOCTEXT_NAMESPACE "FlipbookEditorCommands"
 
 void FFlipbookEditorCommands::RegisterCommands()
 {
diff --git a/Engine/Plugins/2D/Paper2D/Source/Paper2DEditor/Private/PaperEditorCommands.cpp b/Engine/Plugins/2D/Paper2D/Source/Paper2DEditor/Private/PaperEditorCommands.cpp
--- a/Engine/Plugins/2D/Paper2D/Source/Paper2DEditor/Private/PaperEditorCommands.cpp
+++ b/Engine/Plugins/2D/Paper2D/Source/Paper2DEditor/Private/PaperEditorCommands.cpp
@@ -6,7 +6,7 @@
 //////////////////////////////////////////////////////////////////////////
 // FPaperEditorCommands
 
-#define LOCTEXT_NAMESPACE ""
+#define LOCTEXT_NAMESPACE "PaperEditorCommands"
 
 void FPaperEditorCommands::RegisterCommands()
 {
diff --git a/Engine/Plugins/2D/Paper2D/Source/Paper2DEditor/Private/SpriteEditor/SpriteEditorViewportClient.cpp b/Engine/Plugins/2D/Paper2D/Source/Paper2DEditor/Private/SpriteEditor/SpriteEditorViewportClient.cpp
--- a/Engine/Plugins/2D/Paper2D/Source/Paper2DEditor/Private/SpriteEditor/SpriteEditorViewportClient.cpp
+++ b/Engine/Plugins/2D/Paper2D/Source/Paper2DEditor/Private/SpriteEditor/SpriteEditorViewportClient.cpp
@@ -707,11 +707,11 @@ void FSpriteEditorViewportClient::UpdateRelatedSpritesList()
 				continue;
 			}
 
-			const FString* SourceUVString = SpriteAsset.TagsAndValues.Find("SourceUV");
-			const FString* SourceDimensionString = SpriteAsset.TagsAndValues.Find("SourceDimension");
-			FVector2D SourceUV, SourceDimension;
-			if (SourceUVString != nullptr && SourceDimensionString != nullptr)
+			const FString SourceUVString = SpriteAsset.GetTagValueRef<FString>("SourceUV");
+			const FString SourceDimensionString = SpriteAsset.GetTagValueRef<FString>("SourceDimension");
+			if (!SourceUVString.IsEmpty() && !SourceDimensionString.IsEmpty())
 			{
+				FVector2D SourceUV, SourceDimension;
 				if (SourceUV.InitFromString(*SourceUVString) && SourceDimension.InitFromString(*SourceDimensionString))
 				{
 					FRelatedSprite RelatedSprite;
diff --git a/Engine/Plugins/Developer/GitSourceControl/Source/GitSourceControl/Private/GitSourceControlModule.cpp b/Engine/Plugins/Developer/GitSourceControl/Source/GitSourceControl/Private/GitSourceControlModule.cpp
--- a/Engine/Plugins/Developer/GitSourceControl/Source/GitSourceControl/Private/GitSourceControlModule.cpp
+++ b/Engine/Plugins/Developer/GitSourceControl/Source/GitSourceControl/Private/GitSourceControlModule.cpp
@@ -25,8 +25,7 @@ void FGitSourceControlModule::StartupModule()
 	GitSourceControlProvider.RegisterWorker( "MarkForAdd", FGetGitSourceControlWorker::CreateStatic( &CreateWorker<FGitMarkForAddWorker> ) );
 	GitSourceControlProvider.RegisterWorker( "Delete", FGetGitSourceControlWorker::CreateStatic( &CreateWorker<FGitDeleteWorker> ) );
 	GitSourceControlProvider.RegisterWorker( "Revert", FGetGitSourceControlWorker::CreateStatic( &CreateWorker<FGitRevertWorker> ) );
-	// @todo: git fetch remote(s) to be able to show files not up-to-date with the serveur
-//	GitSourceControlProvider.RegisterWorker( "Sync", FGetGitSourceControlWorker::CreateStatic( &CreateWorker<FGitSyncWorker> ) );
+	GitSourceControlProvider.RegisterWorker( "Sync", FGetGitSourceControlWorker::CreateStatic( &CreateWorker<FGitSyncWorker> ) );
 	GitSourceControlProvider.RegisterWorker( "CheckIn", FGetGitSourceControlWorker::CreateStatic( &CreateWorker<FGitCheckInWorker> ) );
 	GitSourceControlProvider.RegisterWorker( "Copy", FGetGitSourceControlWorker::CreateStatic( &CreateWorker<FGitCopyWorker> ) );
 	GitSourceControlProvider.RegisterWorker( "Resolve", FGetGitSourceControlWorker::CreateStatic( &CreateWorker<FGitResolveWorker> ) );
diff --git a/Engine/Plugins/Developer/GitSourceControl/Source/GitSourceControl/Private/GitSourceControlModule.h b/Engine/Plugins/Developer/GitSourceControl/Source/GitSourceControl/Private/GitSourceControlModule.h
--- a/Engine/Plugins/Developer/GitSourceControl/Source/GitSourceControl/Private/GitSourceControlModule.h
+++ b/Engine/Plugins/Developer/GitSourceControl/Source/GitSourceControl/Private/GitSourceControlModule.h
@@ -14,35 +14,35 @@ Written and contributed by [REDACTED_PERSON] ([REDACTED_EMAIL])
 
 ### Supported features
 - initialize a new Git local repository ('git init') to manage your UE4 Game Project.
+- create an appropriate .gitignore file as part as initialization
+- can also make the initial commit
 - display status icons to show modified/added/deleted/untracked files
 - show history of a file
-- diff against depot or between previous versions of a file
+- visual diff of a blueprint against depot or between previous versions of a file
 - revert modifications of a file
-- add a file
-- delete a file
+- add, delete, rename a file
 - checkin/commit a file (cannot handle atomically more than 20 files)
+- migrate an asset between to projects if both are using Git
+- solve a merge conflict on a blueprint
 - show current branch name in status text
-- merge blueprints to solve a merge or rebase conflict
+- Sync to Pull the current branch if there is no local modified files
+- git LFS (Github & Gitlab), git-annex and/or git-media are working under Windows
+- Windows, Mac and Linux
 
 ### What *cannot* be done presently
 - tags: implement ISourceControlLabel to manage git tags
-- .uproject file state si not visible in the current Editor
 - Branch is not in the current Editor workflow (but on Epic Roadmap)
 - Pull/Fetch/Push are not in the current Editor workflow
 - Amend a commit is not in the current Editor workflow
 - configure user name & email ('git config user.name' & git config user.email')
-- Linux is not supported by the SourceControlProvider (#define SOURCE_CONTROL_WITH_SLATE	!PLATFORM_LINUX)
-- git-annex and/or git-media
 
 ### Known issues:
-- global menu "Submit to source control" leads to many lines of Logs "is outside repository"
 - the Editor does not show deleted files (only when deleted externaly?)
 - the Editor does not show missing files
+- the Editor does not show .uproject file state
 - missing localisation for git specific messages
-- migrate an asset should add it to the destination project if also under Git (needs management of 'out of tree' files)
 - displaying states of 'Engine' assets (also needs management of 'out of tree' files)
 - issue #22: A Move/Rename leaves a redirector file behind
-- issue #11: Add the "Resolve" operation introduced in Editor 4.3
 - reverting an asset does not seem to update content in Editor! Issue in Editor?
 - renaming a Blueprint in Editor leaves a tracker file, AND modify too much the asset to enable git to track its history through renaming
 - file history show Changelist as signed integer instead of hexadecimal SHA1
@@ -77,4 +77,4 @@ private:
 
 	/** The settings for Git source control */
 	FGitSourceControlSettings GitSourceControlSettings;
-};
\ No newline at end of file
+};
diff --git a/Engine/Plugins/Developer/GitSourceControl/Source/GitSourceControl/Private/GitSourceControlOperations.cpp b/Engine/Plugins/Developer/GitSourceControl/Source/GitSourceControl/Private/GitSourceControlOperations.cpp
--- a/Engine/Plugins/Developer/GitSourceControl/Source/GitSourceControl/Private/GitSourceControlOperations.cpp
+++ b/Engine/Plugins/Developer/GitSourceControl/Source/GitSourceControl/Private/GitSourceControlOperations.cpp
@@ -169,6 +169,32 @@ bool FGitRevertWorker::UpdateStates() const
 	return GitSourceControlUtils::UpdateCachedStates(States);
 }
 
+FName FGitSyncWorker::GetName() const
+{
+   return "Sync";
+}
+
+bool FGitSyncWorker::Execute(FGitSourceControlCommand& InCommand)
+{
+   // pull the branch to get remote changes by rebasing any local commits (not merging them to avoid complex graphs)
+   // (this cannot work if any local files are modified but not commited)
+   {
+      TArray<FString> Parameters;
+      Parameters.Add(TEXT("--rebase"));
+      InCommand.bCommandSuccessful = GitSourceControlUtils::RunCommand(TEXT("pull"), InCommand.PathToGitBinary, InCommand.PathToRepositoryRoot, Parameters, TArray<FString>(), InCommand.InfoMessages, InCommand.ErrorMessages);
+   }
+
+   // now update the status of our files
+   GitSourceControlUtils::RunUpdateStatus(InCommand.PathToGitBinary, InCommand.PathToRepositoryRoot, InCommand.Files, InCommand.ErrorMessages, States);
+
+   return InCommand.bCommandSuccessful;
+}
+
+bool FGitSyncWorker::UpdateStates() const
+{
+   return GitSourceControlUtils::UpdateCachedStates(States);
+}
+
 FName FGitUpdateStatusWorker::GetName() const
 {
 	return "UpdateStatus";
diff --git a/Engine/Plugins/Developer/GitSourceControl/Source/GitSourceControl/Private/GitSourceControlOperations.h b/Engine/Plugins/Developer/GitSourceControl/Source/GitSourceControl/Private/GitSourceControlOperations.h
--- a/Engine/Plugins/Developer/GitSourceControl/Source/GitSourceControl/Private/GitSourceControlOperations.h
+++ b/Engine/Plugins/Developer/GitSourceControl/Source/GitSourceControl/Private/GitSourceControlOperations.h
@@ -78,7 +78,7 @@ public:
 	TArray<FGitSourceControlState> States;
 };
 
-/** @todo: git fetch remote(s) to be able to show files not up-to-date with the serveur
+/** Git pull --rebase to update branch from its configure remote */
 class FGitSyncWorker : public IGitSourceControlWorker
 {
 public:
@@ -92,7 +92,6 @@ public:
 	/// Map of filenames to Git state
 	TArray<FGitSourceControlState> States;
 };
-*/
 
 /** Get source control status of files on local working copy. */
 class FGitUpdateStatusWorker : public IGitSourceControlWorker
diff --git a/Engine/Plugins/Developer/GitSourceControl/Source/GitSourceControl/Private/SGitSourceControlSettings.cpp b/Engine/Plugins/Developer/GitSourceControl/Source/GitSourceControl/Private/SGitSourceControlSettings.cpp
--- a/Engine/Plugins/Developer/GitSourceControl/Source/GitSourceControl/Private/SGitSourceControlSettings.cpp
+++ b/Engine/Plugins/Developer/GitSourceControl/Source/GitSourceControl/Private/SGitSourceControlSettings.cpp
@@ -12,6 +12,9 @@ void SGitSourceControlSettings::Construct(const FArguments& InArgs)
 	FSlateFontInfo Font = FEditorStyle::GetFontStyle(TEXT("SourceControl.LoginWindow.Font"));
 
 	bAutoCreateGitIgnore = true;
+	bAutoInitialCommit = true;
+
+	InitialCommitMessage = LOCTEXT("InitialCommitMessage", "Initial commit");
 
 	ChildSlot
 	[
@@ -21,7 +24,7 @@ void SGitSourceControlSettings::Construct(const FArguments& InArgs)
 		[
 			SNew(SVerticalBox)
 			+SVerticalBox::Slot()
-			.FillHeight(1.0f)
+			.FillHeight(1.5f)
 			.Padding(2.0f)
 			.VAlign(VAlign_Center)
 			[
@@ -46,7 +49,7 @@ void SGitSourceControlSettings::Construct(const FArguments& InArgs)
 				[
 					SNew(SVerticalBox)
 					+SVerticalBox::Slot()
-					.FillHeight(1.0f)
+					.FillHeight(1.5f)
 					.Padding(2.0f)
 					[
 						SNew(SEditableTextBox)
@@ -173,32 +176,45 @@ void SGitSourceControlSettings::Construct(const FArguments& InArgs)
 				SNew(SHorizontalBox)
 				.Visibility(this, &SGitSourceControlSettings::CanInitializeGitRepository)
 				+SHorizontalBox::Slot()
-				.FillWidth(1.0f)
+				.FillWidth(0.1f)
 				[
 					SNew(SVerticalBox)
 					+SVerticalBox::Slot()
-					.FillHeight(2.0f)
+					.FillHeight(1.0f)
 					.Padding(2.0f)
 					.VAlign(VAlign_Center)
-					.AutoHeight()
 					[
-						SNew(SButton)
-						.Text(LOCTEXT("GitInitRepository", "Initialize project with Git"))
-						.ToolTipText(LOCTEXT("GitInitRepository_Tooltip", "Initialize current project as a new Git repository"))
-						.OnClicked(this, &SGitSourceControlSettings::OnClickedInitializeGitRepository)
-						.HAlign(HAlign_Center)
+						SNew(SCheckBox)
+						.ToolTipText(LOCTEXT("CreateGitIgnore_Tooltip", "Create and add a standard '.gitignore' file"))
+						.IsChecked(ECheckBoxState::Checked)
+						.OnCheckStateChanged(this, &SGitSourceControlSettings::OnCheckedCreateGitIgnore)
+					]
+				]
+				+SHorizontalBox::Slot()
+				.FillWidth(2.9f)
+				[
+					SNew(SVerticalBox)
+					+SVerticalBox::Slot()
+					.FillHeight(1.0f)
+					.Padding(2.0f)
+					.VAlign(VAlign_Center)
+					[
+						SNew(STextBlock)
+						.Text(LOCTEXT("CreateGitIgnore", "Add a .gitignore file"))
+						.ToolTipText(LOCTEXT("CreateGitIgnore_Tooltip", "Create and add a standard '.gitignore' file"))
+						.Font(Font)
 					]
 				]
 			]
 			+SVerticalBox::Slot()
-			.FillHeight(1.0f)
+			.FillHeight(1.5f)
 			.Padding(2.0f)
 			.VAlign(VAlign_Center)
 			[
 				SNew(SHorizontalBox)
 				.Visibility(this, &SGitSourceControlSettings::CanInitializeGitRepository)
 				+SHorizontalBox::Slot()
-				.FillWidth(0.04f)
+				.FillWidth(0.1f)
 				[
 					SNew(SVerticalBox)
 					+SVerticalBox::Slot()
@@ -207,13 +223,13 @@ void SGitSourceControlSettings::Construct(const FArguments& InArgs)
 					.VAlign(VAlign_Center)
 					[
 						SNew(SCheckBox)
-						.ToolTipText(LOCTEXT("CreateGitIgnore_Tooltip", "Create and add a standard '.gitignore' file"))
+						.ToolTipText(LOCTEXT("InitialGitCommit_Tooltip", "Make the initial Git commit"))
 						.IsChecked(ECheckBoxState::Checked)
-						.OnCheckStateChanged(this, &SGitSourceControlSettings::OnCheckedCreateGitIgnore)
+						.OnCheckStateChanged(this, &SGitSourceControlSettings::OnCheckedInitialCommit)
 					]
 				]
 				+SHorizontalBox::Slot()
-				.FillWidth(0.96f)
+				.FillWidth(0.9f)
 				[
 					SNew(SVerticalBox)
 					+SVerticalBox::Slot()
@@ -222,12 +238,52 @@ void SGitSourceControlSettings::Construct(const FArguments& InArgs)
 					.VAlign(VAlign_Center)
 					[
 						SNew(STextBlock)
-						.Text(LOCTEXT("CreateGitIgnore", "Add a .gitignore file"))
-						.ToolTipText(LOCTEXT("CreateGitIgnore_Tooltip", "Create and add a standard '.gitignore' file"))
+						.Text(LOCTEXT("InitialGitCommit", "Make the initial Git Commit"))
+						.ToolTipText(LOCTEXT("InitialGitCommit_Tooltip", "Make the initial Git commit"))
+						.Font(Font)
+					]
+				]
+				+SHorizontalBox::Slot()
+				.FillWidth(2.0f)
+				[
+					SNew(SVerticalBox)
+					+SVerticalBox::Slot()
+					.FillHeight(1.5f)
+					.Padding(2.0f)
+					[
+						SNew(SEditableTextBox)
+						.Text(this, &SGitSourceControlSettings::GetInitialCommitMessage)
+						.ToolTipText(LOCTEXT("InitialCommitMessage_Tooltip", "Message of initial commit"))
+						.OnTextCommitted(this, &SGitSourceControlSettings::OnInitialCommitMessageCommited)
 						.Font(Font)
 					]
 				]
 			]
+			+SVerticalBox::Slot()
+			.FillHeight(2.0f)
+			.Padding(2.0f)
+			.VAlign(VAlign_Center)
+			[
+				SNew(SHorizontalBox)
+				.Visibility(this, &SGitSourceControlSettings::CanInitializeGitRepository)
+				+SHorizontalBox::Slot()
+				.FillWidth(1.0f)
+				[
+					SNew(SVerticalBox)
+					+SVerticalBox::Slot()
+					.FillHeight(2.0f)
+					.Padding(2.0f)
+					.VAlign(VAlign_Center)
+					.AutoHeight()
+					[
+						SNew(SButton)
+						.Text(LOCTEXT("GitInitRepository", "Initialize project with Git"))
+						.ToolTipText(LOCTEXT("GitInitRepository_Tooltip", "Initialize current project as a new Git repository"))
+						.OnClicked(this, &SGitSourceControlSettings::OnClickedInitializeGitRepository)
+						.HAlign(HAlign_Center)
+					]
+				]
+			]
 		]
 	];
 }
@@ -284,9 +340,13 @@ FReply SGitSourceControlSettings::OnClickedInitializeGitRepository()
 	if(GitSourceControl.GetProvider().IsEnabled())
 	{
 		TArray<FString> ProjectFiles;
-		const FString ProjectFile = FPaths::GetCleanFilename(FPaths::GetProjectFilePath());
-		ProjectFiles.Add(ProjectFile);
-		ProjectFiles.Add(TEXT("Config"));
+		ProjectFiles.Add(FPaths::GetCleanFilename(FPaths::GetProjectFilePath()));
+		ProjectFiles.Add(FPaths::GetCleanFilename(FPaths::GameConfigDir()));
+		ProjectFiles.Add(FPaths::GetCleanFilename(FPaths::GameContentDir()));
+		if(FPaths::DirectoryExists(FPaths::GameSourceDir()))
+		{
+			ProjectFiles.Add(FPaths::GetCleanFilename(FPaths::GameSourceDir()));
+		}
 		if(bAutoCreateGitIgnore)
 		{
 			// Create a standard ".gitignore" file with common patterns for a typical Blueprint & C++ project
@@ -297,8 +357,18 @@ FReply SGitSourceControlSettings::OnClickedInitializeGitRepository()
 				ProjectFiles.Add(TEXT(".gitignore"));
 			}
 		}
-		// Add .uproject and Config/ files (and .gitignore if any)
+		// Add .uproject, Config/, Content/ and Source/ files (and .gitignore if any)
 		GitSourceControlUtils::RunCommand(TEXT("add"), PathToGitBinary, PathToGameDir, TArray<FString>(), ProjectFiles, InfoMessages, ErrorMessages);
+		if(bAutoInitialCommit)
+		{
+			// optionnal initial git commit with custom message
+			TArray<FString> Parameters;
+			FString ParamCommitMsg = TEXT("--message=\"");
+			ParamCommitMsg += InitialCommitMessage.ToString();
+			ParamCommitMsg += TEXT("\"");
+			Parameters.Add(ParamCommitMsg);
+			GitSourceControlUtils::RunCommit(PathToGitBinary, PathToGameDir, Parameters, TArray<FString>(), InfoMessages, ErrorMessages);
+		}
 	}
 	return FReply::Handled();
 }
@@ -308,4 +378,19 @@ void SGitSourceControlSettings::OnCheckedCreateGitIgnore(ECheckBoxState NewCheck
 	bAutoCreateGitIgnore = (NewCheckedState == ECheckBoxState::Checked);
 }
 
+void SGitSourceControlSettings::OnCheckedInitialCommit(ECheckBoxState NewCheckedState)
+{
+	bAutoInitialCommit = (NewCheckedState == ECheckBoxState::Checked);
+}
+
+void SGitSourceControlSettings::OnInitialCommitMessageCommited(const FText& InText, ETextCommit::Type InCommitType)
+{
+	InitialCommitMessage = InText;
+}
+
+FText SGitSourceControlSettings::GetInitialCommitMessage() const
+{
+	return InitialCommitMessage;
+}
+
 #undef LOCTEXT_NAMESPACE
diff --git a/Engine/Plugins/Developer/GitSourceControl/Source/GitSourceControl/Private/SGitSourceControlSettings.h b/Engine/Plugins/Developer/GitSourceControl/Source/GitSourceControl/Private/SGitSourceControlSettings.h
--- a/Engine/Plugins/Developer/GitSourceControl/Source/GitSourceControl/Private/SGitSourceControlSettings.h
+++ b/Engine/Plugins/Developer/GitSourceControl/Source/GitSourceControl/Private/SGitSourceControlSettings.h
@@ -30,6 +30,13 @@ private:
 	/** Delegate to initialize a new Git repository */
 	EVisibility CanInitializeGitRepository() const;
 	FReply OnClickedInitializeGitRepository();
+
 	void OnCheckedCreateGitIgnore(ECheckBoxState NewCheckedState);
 	bool bAutoCreateGitIgnore;
+
+	void OnCheckedInitialCommit(ECheckBoxState NewCheckedState);
+	bool bAutoInitialCommit;
+	void OnInitialCommitMessageCommited(const FText& InText, ETextCommit::Type InCommitType);
+	FText GetInitialCommitMessage() const;
+	FText InitialCommitMessage;
 };
\ No newline at end of file
diff --git a/Engine/Plugins/Editor/PluginBrowser/Source/PluginBrowser/PluginBrowser.Build.cs b/Engine/Plugins/Editor/PluginBrowser/Source/PluginBrowser/PluginBrowser.Build.cs
--- a/Engine/Plugins/Editor/PluginBrowser/Source/PluginBrowser/PluginBrowser.Build.cs
+++ b/Engine/Plugins/Editor/PluginBrowser/Source/PluginBrowser/PluginBrowser.Build.cs
@@ -26,7 +26,8 @@ namespace UnrealBuildTool.Rules
 					"SharedSettingsWidgets",
 					"DirectoryWatcher",
 					"GameProjectGeneration",
-					"MainFrame"
+					"MainFrame",
+                    "UATHelper"
 				}
 			);
 
diff --git a/Engine/Plugins/Editor/PluginBrowser/Source/PluginBrowser/Private/SPluginTile.cpp b/Engine/Plugins/Editor/PluginBrowser/Source/PluginBrowser/Private/SPluginTile.cpp
--- a/Engine/Plugins/Editor/PluginBrowser/Source/PluginBrowser/Private/SPluginTile.cpp
+++ b/Engine/Plugins/Editor/PluginBrowser/Source/PluginBrowser/Private/SPluginTile.cpp
@@ -13,6 +13,9 @@
 #include "PluginBrowserModule.h"
 #include "ISourceControlModule.h"
 #include "PropertyEditorModule.h"
+#include "IUATHelperModule.h"
+#include "IMainFrameModule.h"
+#include "DesktopPlatformModule.h"
 
 #define LOCTEXT_NAMESPACE "PluginListTile"
 
@@ -312,14 +315,24 @@ void SPluginTile::RecreateWidgets()
 														SNew(SHorizontalBox)
 
 														+ SHorizontalBox::Slot()
-															.AutoWidth()
-															.Padding(PaddingAmount)
-															[
-																SNew(SHyperlink)
-																	.Visibility(this, &SPluginTile::GetAuthoringButtonsVisibility)	
-																	.OnNavigate(this, &SPluginTile::OnEditPlugin)
-																	.Text(LOCTEXT("EditPlugin", "Edit..."))
-															]
+														.AutoWidth()
+														.Padding(PaddingAmount)
+														[
+															SNew(SHyperlink)
+															.Visibility(this, &SPluginTile::GetAuthoringButtonsVisibility)	
+															.OnNavigate(this, &SPluginTile::OnEditPlugin)
+															.Text(LOCTEXT("EditPlugin", "Edit..."))
+														]
+
+														+ SHorizontalBox::Slot()
+														.AutoWidth()
+														.Padding(PaddingAmount)
+														[
+															SNew(SHyperlink)
+															.Visibility(this, &SPluginTile::GetAuthoringButtonsVisibility)
+															.OnNavigate(this, &SPluginTile::OnPackagePlugin)
+															.Text(LOCTEXT("PackagePlugin", "Package..."))
+														]
 													]
 
 												// support link
@@ -523,4 +536,40 @@ FReply SPluginTile::OnEditPluginFinished(UPluginMetadataObject* MetadataObject)
 	return FReply::Handled();
 }
 
+void SPluginTile::OnPackagePlugin()
+{
+	void* ParentWindowWindowHandle = nullptr;
+	IMainFrameModule& MainFrameModule = FModuleManager::LoadModuleChecked<IMainFrameModule>(TEXT("MainFrame"));
+	const TSharedPtr<SWindow>& MainFrameParentWindow = MainFrameModule.GetParentWindow();
+	if ( MainFrameParentWindow.IsValid() && MainFrameParentWindow->GetNativeWindow().IsValid() )
+	{
+		ParentWindowWindowHandle = MainFrameParentWindow->GetNativeWindow()->GetOSWindowHandle();
+	}
+
+	FString DefaultDirectory;
+	FString OutputDirectory;
+
+	if ( !FDesktopPlatformModule::Get()->OpenDirectoryDialog(ParentWindowWindowHandle, LOCTEXT("PackagePluginDialogTitle", "Package Plugin...").ToString(), DefaultDirectory, OutputDirectory) )
+	{
+		return;
+	}
+
+	// Ensure path is full rather than relative (for macs)
+	FString UPluginFilename = Plugin->GetDescriptorFileName();
+	FString CommandLine = FString::Printf(TEXT("BuildPlugin -Rocket -Plugin=\"%s\" -Package=\"%s\""), *UPluginFilename, *OutputDirectory);
+
+#if PLATFORM_WINDOWS
+	FText PlatformName = LOCTEXT("PlatformName_Windows", "Windows");
+#elif PLATFORM_MAC
+	FText PlatformName = LOCTEXT("PlatformName_Mac", "Mac");
+#elif PLATFORM_LINUX
+	FText PlatformName = LOCTEXT("PlatformName_Linux", "Linux");
+#else
+	FText PlatformName = LOCTEXT("PlatformName_Other", "Other OS");
+#endif
+
+	IUATHelperModule::Get().CreateUatTask(CommandLine, PlatformName, LOCTEXT("PackagePluginTaskName", "Packaging Plugin"),
+		LOCTEXT("PackagePluginTaskShortName", "Package Plugin Task"), FEditorStyle::GetBrush(TEXT("MainFrame.CookContent")));
+}
+
 #undef LOCTEXT_NAMESPACE
diff --git a/Engine/Plugins/Editor/PluginBrowser/Source/PluginBrowser/Private/SPluginTile.h b/Engine/Plugins/Editor/PluginBrowser/Source/PluginBrowser/Private/SPluginTile.h
--- a/Engine/Plugins/Editor/PluginBrowser/Source/PluginBrowser/Private/SPluginTile.h
+++ b/Engine/Plugins/Editor/PluginBrowser/Source/PluginBrowser/Private/SPluginTile.h
@@ -39,6 +39,9 @@ private:
 	/** Called when the 'edit' hyperlink is clicked */
 	FReply OnEditPluginFinished(UPluginMetadataObject* MetadataObject);
 
+	/** Called when the 'package' hyperlink is clicked */
+	void OnPackagePlugin();
+
 private:
 
 	/** The item we're representing the in tree */
diff --git a/Engine/Plugins/Editor/PluginBrowser/Templates/EditorMode/Source/PLUGIN_NAME/Private/PLUGIN_NAMEEdModeToolkit.cpp b/Engine/Plugins/Editor/PluginBrowser/Templates/EditorMode/Source/PLUGIN_NAME/Private/PLUGIN_NAMEEdModeToolkit.cpp
--- a/Engine/Plugins/Editor/PluginBrowser/Templates/EditorMode/Source/PLUGIN_NAME/Private/PLUGIN_NAMEEdModeToolkit.cpp
+++ b/Engine/Plugins/Editor/PluginBrowser/Templates/EditorMode/Source/PLUGIN_NAME/Private/PLUGIN_NAMEEdModeToolkit.cpp
@@ -7,6 +7,10 @@
 #define LOCTEXT_NAMESPACE "FPLUGIN_NAMEEdModeToolkit"
 
 FPLUGIN_NAMEEdModeToolkit::FPLUGIN_NAMEEdModeToolkit()
+{
+}
+
+void FPLUGIN_NAMEEdModeToolkit::Init(const TSharedPtr<IToolkitHost>& InitToolkitHost)
 {
 	struct Locals
 	{
@@ -95,6 +99,8 @@ FPLUGIN_NAMEEdModeToolkit::FPLUGIN_NAMEEdModeToolkit()
 				]
 
 		];
+		
+	FModeToolkit::Init(InitToolkitHost);
 }
 
 FName FPLUGIN_NAMEEdModeToolkit::GetToolkitFName() const
diff --git a/Engine/Plugins/Editor/PluginBrowser/Templates/EditorMode/Source/PLUGIN_NAME/Public/PLUGIN_NAMEEdModeToolkit.h b/Engine/Plugins/Editor/PluginBrowser/Templates/EditorMode/Source/PLUGIN_NAME/Public/PLUGIN_NAMEEdModeToolkit.h
--- a/Engine/Plugins/Editor/PluginBrowser/Templates/EditorMode/Source/PLUGIN_NAME/Public/PLUGIN_NAMEEdModeToolkit.h
+++ b/Engine/Plugins/Editor/PluginBrowser/Templates/EditorMode/Source/PLUGIN_NAME/Public/PLUGIN_NAMEEdModeToolkit.h
@@ -10,6 +10,9 @@ class FPLUGIN_NAMEEdModeToolkit : public FModeToolkit
 public:
 
 	FPLUGIN_NAMEEdModeToolkit();
+	
+	/** FModeToolkit interface */
+	virtual void Init(const TSharedPtr<IToolkitHost>& InitToolkitHost) override;
 
 	/** IToolkit interface */
 	virtual FName GetToolkitFName() const override;
diff --git a/Engine/Plugins/Media/MediaPlayerEditor/Source/MediaPlayerEditor/Private/Models/MediaPlayerEditorCommands.cpp b/Engine/Plugins/Media/MediaPlayerEditor/Source/MediaPlayerEditor/Private/Models/MediaPlayerEditorCommands.cpp
--- a/Engine/Plugins/Media/MediaPlayerEditor/Source/MediaPlayerEditor/Private/Models/MediaPlayerEditorCommands.cpp
+++ b/Engine/Plugins/Media/MediaPlayerEditor/Source/MediaPlayerEditor/Private/Models/MediaPlayerEditorCommands.cpp
@@ -2,7 +2,7 @@
 
 #include "MediaPlayerEditorPrivatePCH.h"
 
-#define LOCTEXT_NAMESPACE ""
+#define LOCTEXT_NAMESPACE "MediaPlayerEditorCommands"
 
 void FMediaPlayerEditorCommands::RegisterCommands()
 {
diff --git a/Engine/Plugins/MovieScene/LevelSequenceEditor/Source/LevelSequenceEditor/Private/LevelSequenceEditorToolkit.cpp b/Engine/Plugins/MovieScene/LevelSequenceEditor/Source/LevelSequenceEditor/Private/LevelSequenceEditorToolkit.cpp
--- a/Engine/Plugins/MovieScene/LevelSequenceEditor/Source/LevelSequenceEditor/Private/LevelSequenceEditorToolkit.cpp
+++ b/Engine/Plugins/MovieScene/LevelSequenceEditor/Source/LevelSequenceEditor/Private/LevelSequenceEditorToolkit.cpp
@@ -6,7 +6,7 @@
 #include "Editor/LevelEditor/Public/ILevelViewport.h"
 #include "Editor/WorkspaceMenuStructure/Public/WorkspaceMenuStructure.h"
 #include "Editor/WorkspaceMenuStructure/Public/WorkspaceMenuStructureModule.h"
-#include "ClassIconFinder.h"
+#include "SlateIconFinder.h"
 #include "ISequencer.h"
 #include "ISequencerModule.h"
 #include "MovieScene.h"
@@ -485,13 +485,11 @@ void FLevelSequenceEditorToolkit::AddPosessActorMenuExtensions(FMenuBuilder& Men
 	ActorsValidForPossession.RemoveAll([&](AActor* In){ return !IsActorValidForPossession(In); });
 
 	FText SelectedLabel;
-	FSlateIcon ActorIcon(FEditorStyle::GetStyleSetName(), FClassIconFinder::FindIconNameForClass(AActor::StaticClass()));
-
+	FSlateIcon ActorIcon = FSlateIconFinder::FindIconForClass(AActor::StaticClass());
 	if (ActorsValidForPossession.Num() == 1)
 	{
 		SelectedLabel = FText::Format(LOCTEXT("AddSpecificActor", "Add '{0}'"), FText::FromString(ActorsValidForPossession[0]->GetActorLabel()));
-		FName IconName = FClassIconFinder::FindIconNameForActor(ActorsValidForPossession[0]);
-		ActorIcon = FSlateIcon(FEditorStyle::GetStyleSetName(), FClassIconFinder::FindIconNameForClass(ActorsValidForPossession[0]->GetClass()));
+		ActorIcon = FSlateIconFinder::FindIconForClass(ActorsValidForPossession[0]->GetClass());
 	}
 	else if (ActorsValidForPossession.Num() > 1)
 	{
diff --git a/Engine/Plugins/NetcodeUnitTest/NetcodeUnitTest/Source/NetcodeUnitTest/Private/UI/LogWidgetCommands.cpp b/Engine/Plugins/NetcodeUnitTest/NetcodeUnitTest/Source/NetcodeUnitTest/Private/UI/LogWidgetCommands.cpp
--- a/Engine/Plugins/NetcodeUnitTest/NetcodeUnitTest/Source/NetcodeUnitTest/Private/UI/LogWidgetCommands.cpp
+++ b/Engine/Plugins/NetcodeUnitTest/NetcodeUnitTest/Source/NetcodeUnitTest/Private/UI/LogWidgetCommands.cpp
@@ -4,7 +4,7 @@
 #include "LogWidgetCommands.h"
 #include "Framework/Commands/UICommandInfo.h"
 
-#define LOCTEXT_NAMESPACE ""
+#define LOCTEXT_NAMESPACE "LogWidgetCommands"
 
 void FLogWidgetCommands::RegisterCommands()
 {
diff --git a/Engine/Plugins/Runtime/SoundMod/Source/SoundModImporter/SoundModImporter.Build.cs b/Engine/Plugins/Runtime/SoundMod/Source/SoundModImporter/SoundModImporter.Build.cs
--- a/Engine/Plugins/Runtime/SoundMod/Source/SoundModImporter/SoundModImporter.Build.cs
+++ b/Engine/Plugins/Runtime/SoundMod/Source/SoundModImporter/SoundModImporter.Build.cs
@@ -11,7 +11,8 @@ public class SoundModImporter : ModuleRules
 				"Core",
 				"CoreUObject",
 				"Engine",
-				"Slate",
+				"SlateCore",
+                "Slate",
 				"SoundMod",
 				"UnrealEd"
 			});
diff --git a/Engine/Source/Developer/AssetTools/Private/AssetTypeActions/AssetTypeActions_AnimBlueprint.cpp b/Engine/Source/Developer/AssetTools/Private/AssetTypeActions/AssetTypeActions_AnimBlueprint.cpp
--- a/Engine/Source/Developer/AssetTools/Private/AssetTypeActions/AssetTypeActions_AnimBlueprint.cpp
+++ b/Engine/Source/Developer/AssetTools/Private/AssetTypeActions/AssetTypeActions_AnimBlueprint.cpp
@@ -7,7 +7,7 @@
 #include "SBlueprintDiff.h"
 #include "SNotificationList.h"
 #include "SSkeletonWidget.h"
-#include "ClassIconFinder.h"
+#include "SlateIconFinder.h"
 
 #define LOCTEXT_NAMESPACE "AssetTypeActions"
 
@@ -251,7 +251,7 @@ void FAssetTypeActions_AnimBlueprint::RetargetAssets(TArray<UObject*> InAnimBlue
 
 TSharedPtr<SWidget> FAssetTypeActions_AnimBlueprint::GetThumbnailOverlay(const FAssetData& AssetData) const
 {
-	const FSlateBrush* Icon = FClassIconFinder::FindIconForClass(UAnimBlueprint::StaticClass());
+	const FSlateBrush* Icon = FSlateIconFinder::FindIconBrushForClass(UAnimBlueprint::StaticClass());
 
 	return SNew(SBorder)
 		.BorderImage(FEditorStyle::GetNoBrush())
diff --git a/Engine/Source/Developer/AssetTools/Private/AssetTypeActions/AssetTypeActions_Blueprint.cpp b/Engine/Source/Developer/AssetTools/Private/AssetTypeActions/AssetTypeActions_Blueprint.cpp
--- a/Engine/Source/Developer/AssetTools/Private/AssetTypeActions/AssetTypeActions_Blueprint.cpp
+++ b/Engine/Source/Developer/AssetTools/Private/AssetTypeActions/AssetTypeActions_Blueprint.cpp
@@ -76,7 +76,7 @@ void FAssetTypeActions_Blueprint::OpenAssetEditor( const TArray<UObject*>& InObj
 			bool bLetOpen = true;
 			if (!Blueprint->SkeletonGeneratedClass || !Blueprint->GeneratedClass)
 			{
-				bLetOpen = EAppReturnType::Yes == FMessageDialog::Open(EAppMsgType::YesNoCancel, LOCTEXT("FailedToLoadBlueprint", "Blueprint could not be loaded because it derives from an invalid class.  Check to make sure the parent class for this blueprint hasn't been removed! Do you want to continue (it can crash the editor)?"));
+				bLetOpen = EAppReturnType::Yes == FMessageDialog::Open(EAppMsgType::YesNoCancel, LOCTEXT("FailedToLoadBlueprintWithContinue", "Blueprint could not be loaded because it derives from an invalid class.  Check to make sure the parent class for this blueprint hasn't been removed! Do you want to continue (it can crash the editor)?"));
 			}
 			if (bLetOpen)
 			{
@@ -284,13 +284,11 @@ UThumbnailInfo* FAssetTypeActions_Blueprint::GetThumbnailInfo(UObject* Asset) co
 
 FText FAssetTypeActions_Blueprint::GetAssetDescription(const FAssetData& AssetData) const
 {
-	if (const FString* pDescription = AssetData.TagsAndValues.Find(GET_MEMBER_NAME_CHECKED(UBlueprint, BlueprintDescription)))
+	FString Description = AssetData.GetTagValueRef<FString>(GET_MEMBER_NAME_CHECKED(UBlueprint, BlueprintDescription));
+	if (!Description.IsEmpty())
 	{
-		if (!pDescription->IsEmpty())
-		{
-			const FString DescriptionStr(*pDescription);
-			return FText::FromString(DescriptionStr.Replace(TEXT("\\n"), TEXT("\n")));
-		}
+		Description.ReplaceInline(TEXT("\\n"), TEXT("\n"));
+		return FText::FromString(MoveTemp(Description));
 	}
 
 	return FText::GetEmpty();
@@ -303,15 +301,14 @@ TWeakPtr<IClassTypeActions> FAssetTypeActions_Blueprint::GetClassTypeActions(con
 
 	// Blueprints get the class type actions for their parent native class - this avoids us having to load the blueprint
 	UClass* ParentClass = nullptr;
-	const FString* ParentClassNamePtr = AssetData.TagsAndValues.Find(NativeParentClassTag);
-	if(!ParentClassNamePtr)
+	FString ParentClassName;
+	if(!AssetData.GetTagValue(NativeParentClassTag, ParentClassName))
 	{
-		ParentClassNamePtr = AssetData.TagsAndValues.Find(ParentClassTag);
+		AssetData.GetTagValue(ParentClassTag, ParentClassName);
 	}
-	if(ParentClassNamePtr && !ParentClassNamePtr->IsEmpty())
+	if(!ParentClassName.IsEmpty())
 	{
 		UObject* Outer = nullptr;
-		FString ParentClassName = *ParentClassNamePtr;
 		ResolveName(Outer, ParentClassName, false, false);
 		ParentClass = FindObject<UClass>(ANY_PACKAGE, *ParentClassName);
 	}
diff --git a/Engine/Source/Developer/AssetTools/Private/AssetTypeActions/AssetTypeActions_StaticMesh.cpp b/Engine/Source/Developer/AssetTools/Private/AssetTypeActions/AssetTypeActions_StaticMesh.cpp
--- a/Engine/Source/Developer/AssetTools/Private/AssetTypeActions/AssetTypeActions_StaticMesh.cpp
+++ b/Engine/Source/Developer/AssetTools/Private/AssetTypeActions/AssetTypeActions_StaticMesh.cpp
@@ -141,7 +141,7 @@ void FAssetTypeActions_StaticMesh::GetLODMenu(class FMenuBuilder& MenuBuilder, T
 		)
 		);
 
-	FText PasteLabel = FText(LOCTEXT("StaticMesh_PasteLODWithName", "Paste LOD"));
+	FText PasteLabel = FText(LOCTEXT("StaticMesh_PasteLOD", "Paste LOD"));
 	if (LODCopyMesh.IsValid())
 	{
 		PasteLabel = FText::Format(LOCTEXT("StaticMesh_PasteLODWithName", "Paste LOD from {0}"), FText::FromString(LODCopyMesh->GetName()));
diff --git a/Engine/Source/Developer/AssetTools/Private/AssetTypeActions/AssetTypeActions_Struct.cpp b/Engine/Source/Developer/AssetTools/Private/AssetTypeActions/AssetTypeActions_Struct.cpp
--- a/Engine/Source/Developer/AssetTools/Private/AssetTypeActions/AssetTypeActions_Struct.cpp
+++ b/Engine/Source/Developer/AssetTools/Private/AssetTypeActions/AssetTypeActions_Struct.cpp
@@ -29,13 +29,11 @@ void FAssetTypeActions_Struct::OpenAssetEditor( const TArray<UObject*>& InObject
 
 FText FAssetTypeActions_Struct::GetAssetDescription(const FAssetData& AssetData) const
 {
-	if (const FString* pDescription = AssetData.TagsAndValues.Find(TEXT("Tooltip")))
+	FString Description = AssetData.GetTagValueRef<FString>("Tooltip");
+	if (!Description.IsEmpty())
 	{
-		if (!pDescription->IsEmpty())
-		{
-			const FString DescriptionStr(*pDescription);
-			return FText::FromString(DescriptionStr.Replace(TEXT("\\n"), TEXT("\n")));
-		}
+		Description.ReplaceInline(TEXT("\\n"), TEXT("\n"));
+		return FText::FromString(MoveTemp(Description));
 	}
 
 	return FText::GetEmpty();
diff --git a/Engine/Source/Developer/AutomationController/Private/AutomationCommandline.cpp b/Engine/Source/Developer/AutomationController/Private/AutomationCommandline.cpp
--- a/Engine/Source/Developer/AutomationController/Private/AutomationCommandline.cpp
+++ b/Engine/Source/Developer/AutomationController/Private/AutomationCommandline.cpp
@@ -55,7 +55,7 @@ public:
 
 
 		// Register for the callback that tells us there are tests available
-		AutomationController->OnTestsRefreshed().BindRaw(this, &FAutomationExecCmd::HandleRefreshTestCallback);
+		AutomationController->OnTestsRefreshed().AddRaw(this, &FAutomationExecCmd::HandleRefreshTestCallback);
 
 		TickHandler = FTicker::GetCoreTicker().AddTicker(FTickerDelegate::CreateRaw(this, &FAutomationExecCmd::Tick));
 
@@ -78,8 +78,12 @@ public:
 	}
 	void Shutdown()
 	{
-		//IAutomationControllerModule* AutomationControllerModule = &FModuleManager::LoadModuleChecked<IAutomationControllerModule>("AutomationController");
-		//AutomationControllerModule->ShutdownModule();
+		IAutomationControllerModule* AutomationControllerModule = FModuleManager::GetModulePtr<IAutomationControllerModule>("AutomationController");
+		if ( AutomationControllerModule )
+		{
+			AutomationController = AutomationControllerModule->GetAutomationController();
+			AutomationController->OnTestsRefreshed().RemoveAll(this);
+		}
 
 		FTicker::GetCoreTicker().RemoveTicker(TickHandler);
 	}
diff --git a/Engine/Source/Developer/AutomationController/Private/AutomationControllerManager.h b/Engine/Source/Developer/AutomationController/Private/AutomationControllerManager.h
--- a/Engine/Source/Developer/AutomationController/Private/AutomationControllerManager.h
+++ b/Engine/Source/Developer/AutomationController/Private/AutomationControllerManager.h
@@ -185,6 +185,11 @@ public:
 		return TestsRefreshedDelegate;
 	}
 
+	virtual FOnAutomationControllerTestsComplete& OnTestsComplete() override
+	{
+		return TestsCompleteDelegate;
+	}
+
 	virtual FOnAutomationControllerReset& OnControllerReset() override
 	{
 		return ControllerResetDelegate;
@@ -195,11 +200,6 @@ public:
 	virtual void ToggleDeviceGroupFlag( EAutomationDeviceGroupTypes::Type InDeviceGroup ) override;
 	virtual void UpdateDeviceGroups() override;
 
-	virtual void SetTestsCompleteCallback(const FOnAutomationControllerTestsComplete& NewCallback) override
-	{
-		TestsCompleteDelegate = NewCallback;
-	}
-
 	virtual void TrackReportHistory(const bool bShouldTrack, const int32 NumReportsToTrack) override;
 	virtual const bool IsTrackingHistory() const override;
 	virtual const int32 GetNumberHistoryItemsTracking() const override;
diff --git a/Engine/Source/Developer/AutomationController/Private/AutomationControllerManger.cpp b/Engine/Source/Developer/AutomationController/Private/AutomationControllerManger.cpp
--- a/Engine/Source/Developer/AutomationController/Private/AutomationControllerManger.cpp
+++ b/Engine/Source/Developer/AutomationController/Private/AutomationControllerManger.cpp
@@ -18,7 +18,7 @@ void FAutomationControllerManager::RequestAvailableWorkers( const FGuid& Session
 	++ExecutionCount;
 	DeviceClusterManager.Reset();
 
-	ControllerResetDelegate.ExecuteIfBound();
+	ControllerResetDelegate.Broadcast();
 
 	// Don't allow reports to be exported
 	bTestResultsAvailable = false;
@@ -190,7 +190,7 @@ void FAutomationControllerManager::ProcessAvailableTasks()
 				{
 					ProcessResults();
 					//Notify the graphical layout we are done processing results.
-					TestsCompleteDelegate.ExecuteIfBound();
+					TestsCompleteDelegate.Broadcast();
 				}
 			}
 		}
@@ -333,17 +333,17 @@ void FAutomationControllerManager::Startup()
 void FAutomationControllerManager::Shutdown()
 {
 	MessageEndpoint.Reset();
-	ShutdownDelegate.ExecuteIfBound();
+	ShutdownDelegate.Broadcast();
 	RemoveCallbacks();
 }
 
 
 void FAutomationControllerManager::RemoveCallbacks()
 {
-	ShutdownDelegate.Unbind();
-	TestsAvailableDelegate.Unbind();
-	TestsRefreshedDelegate.Unbind();
-	TestsCompleteDelegate.Unbind();
+	ShutdownDelegate.Clear();
+	TestsAvailableDelegate.Clear();
+	TestsRefreshedDelegate.Clear();
+	TestsCompleteDelegate.Clear();
 }
 
 
@@ -387,7 +387,7 @@ void FAutomationControllerManager::SetTestNames( const FMessageAddress& Automati
 	// If we have received all the responses we expect to
 	if (RefreshTestResponses == DeviceClusterManager.GetNumClusters())
 	{
-		TestsRefreshedDelegate.ExecuteIfBound();
+		TestsRefreshedDelegate.Broadcast();
 
 		// Update the tests with tracking details
 		ReportManager.TrackHistory(bTrackHistory, NumberOfHistoryItemsTracked);
@@ -457,7 +457,7 @@ void FAutomationControllerManager::SetControllerStatus( EAutomationControllerMod
 	{
 		// Inform the UI if the test state has changed
 		AutomationTestState = InAutomationTestState;
-		TestsAvailableDelegate.ExecuteIfBound(AutomationTestState);
+		TestsAvailableDelegate.Broadcast(AutomationTestState);
 	}
 }
 
diff --git a/Engine/Source/Developer/AutomationController/Private/AutomationReport.cpp b/Engine/Source/Developer/AutomationController/Private/AutomationReport.cpp
--- a/Engine/Source/Developer/AutomationController/Private/AutomationReport.cpp
+++ b/Engine/Source/Developer/AutomationController/Private/AutomationReport.cpp
@@ -202,6 +202,15 @@ uint32 FAutomationReport::GetTestFlags( ) const
 	return TestInfo.GetTestFlags();
 }
 
+FString FAutomationReport::GetSourceFile() const
+{
+	return TestInfo.GetSourceFile();
+}
+
+int32 FAutomationReport::GetSourceFileLine() const
+{
+	return TestInfo.GetSourceFileLine();
+}
 
 void FAutomationReport::SetTestFlags( const uint32 InTestFlags)
 {
diff --git a/Engine/Source/Developer/AutomationController/Private/AutomationReport.h b/Engine/Source/Developer/AutomationController/Private/AutomationReport.h
--- a/Engine/Source/Developer/AutomationController/Private/AutomationReport.h
+++ b/Engine/Source/Developer/AutomationController/Private/AutomationReport.h
@@ -32,6 +32,8 @@ public:
 	virtual bool IsSupported(const int32 ClusterIndex) const override;
 	virtual void SetTestFlags(const uint32 TestFlags) override;
 	virtual uint32 GetTestFlags( ) const override;
+	virtual FString GetSourceFile() const override;
+	virtual int32 GetSourceFileLine() const override;
 	virtual const bool IsParent( ) override;
 	virtual const bool IsSmokeTest( ) override;
 	virtual bool SetFilter( TSharedPtr< AutomationFilterCollection > InFilter, const bool ParentPassedFilter = false ) override;
diff --git a/Engine/Source/Developer/AutomationController/Public/Interfaces/IAutomationControllerManager.h b/Engine/Source/Developer/AutomationController/Public/Interfaces/IAutomationControllerManager.h
--- a/Engine/Source/Developer/AutomationController/Public/Interfaces/IAutomationControllerManager.h
+++ b/Engine/Source/Developer/AutomationController/Public/Interfaces/IAutomationControllerManager.h
@@ -134,19 +134,19 @@ namespace EAutomationDeviceGroupTypes
 
 
 /** Console command delegate type (takes no arguments.)  This is a void callback function. */
-DECLARE_DELEGATE(FOnAutomationControllerReset);
+DECLARE_MULTICAST_DELEGATE(FOnAutomationControllerReset);
 
 /** Console command delegate type (takes no arguments.)  This is a void callback function. */
-DECLARE_DELEGATE(FOnAutomationControllerTestsRefreshed);
+DECLARE_MULTICAST_DELEGATE(FOnAutomationControllerTestsRefreshed);
 
 /** Delegate for setting the automation controller status in the UI */
-DECLARE_DELEGATE_OneParam(FOnAutomationControllerManagerTestsAvailable, EAutomationControllerModuleState::Type);
+DECLARE_MULTICAST_DELEGATE_OneParam(FOnAutomationControllerManagerTestsAvailable, EAutomationControllerModuleState::Type);
 
 /** Delegate to call when the automation controller is shutdown. */
-DECLARE_DELEGATE(FOnAutomationControllerManagerShutdown);
+DECLARE_MULTICAST_DELEGATE(FOnAutomationControllerManagerShutdown);
 
 /** Delegate to call when the tests are complete. */
-DECLARE_DELEGATE(FOnAutomationControllerTestsComplete);
+DECLARE_MULTICAST_DELEGATE(FOnAutomationControllerTestsComplete);
 
 
 /**
@@ -385,11 +385,6 @@ public:
 	/** Updates the clusters when the device grouping changes. */
 	virtual void UpdateDeviceGroups( ) = 0;
 
-	/**
-	 * Sets the test complete callback delegate.
-	 */
-	virtual void SetTestsCompleteCallback(const FOnAutomationControllerTestsComplete& NewCallback) = 0;
-
 	/**
 	 * Dictate whether to save the reports history and the number of history items to track.
 	 *
@@ -411,26 +406,31 @@ public:
 public:
 
 	/**
-	 * Gets a delegate that is invoked when the controller manager shuts down.
+	 * Gets the multicast delegate that is invoked when the controller manager shuts down.
 	 *
 	 * @return The delegate.
 	 */
 	virtual FOnAutomationControllerManagerShutdown& OnShutdown( ) = 0;
 
 	/**
-	 * Gets a delegate that is invoked when the controller has tests available.
+	 * Gets the multicast delegate that is invoked when the controller has tests available.
 	 *
 	 * @return The delegate.
 	 */
 	virtual FOnAutomationControllerManagerTestsAvailable& OnTestsAvailable( ) = 0;
 
 	/**
-	 * Gets a delegate that is invoked when the controller's test status changes.
+	 * Gets the multicast delegate that is invoked when the controller's test status changes.
 	 *
 	 * @return The delegate.
 	 */
 	virtual FOnAutomationControllerTestsRefreshed& OnTestsRefreshed( ) = 0;
 
+	/**
+	 * Gets the test complete multicast delegate.
+	 */
+	virtual FOnAutomationControllerTestsComplete& OnTestsComplete( ) = 0;
+
 	/**
 	* Gets a delegate that is invoked when the controller's reset.
 	*
diff --git a/Engine/Source/Developer/AutomationController/Public/Interfaces/IAutomationReport.h b/Engine/Source/Developer/AutomationController/Public/Interfaces/IAutomationReport.h
--- a/Engine/Source/Developer/AutomationController/Public/Interfaces/IAutomationReport.h
+++ b/Engine/Source/Developer/AutomationController/Public/Interfaces/IAutomationReport.h
@@ -174,6 +174,12 @@ public:
 	 */
 	virtual uint32 GetTestFlags() const = 0;
 
+	/** Gets the source file the test was defined on. */
+	virtual FString GetSourceFile() const = 0;
+
+	/** Gets the source file line number the test was defined on. */
+	virtual int32 GetSourceFileLine() const = 0;
+
 	/** Recursively gets the number of child nodes */
 	virtual int32 GetTotalNumChildren() const = 0;
 
diff --git a/Engine/Source/Developer/AutomationWindow/Private/SAutomationGraphicalResultBox.cpp b/Engine/Source/Developer/AutomationWindow/Private/SAutomationGraphicalResultBox.cpp
--- a/Engine/Source/Developer/AutomationWindow/Private/SAutomationGraphicalResultBox.cpp
+++ b/Engine/Source/Developer/AutomationWindow/Private/SAutomationGraphicalResultBox.cpp
@@ -28,9 +28,13 @@ void SAutomationGraphicalResultBox::Construct( const FArguments& InArgs, const I
 
 	ClearResults();
 
-	AutomationController->SetTestsCompleteCallback(FOnAutomationControllerTestsComplete::CreateRaw(this, &SAutomationGraphicalResultBox::OnFillResults));
+	AutomationController->OnTestsComplete().AddRaw(this, &SAutomationGraphicalResultBox::OnFillResults);
 }
 
+SAutomationGraphicalResultBox::~SAutomationGraphicalResultBox()
+{
+	AutomationController->OnTestsComplete().RemoveAll(this);
+}
 
 void SAutomationGraphicalResultBox::ClearResults()
 {
diff --git a/Engine/Source/Developer/AutomationWindow/Private/SAutomationGraphicalResultBox.h b/Engine/Source/Developer/AutomationWindow/Private/SAutomationGraphicalResultBox.h
--- a/Engine/Source/Developer/AutomationWindow/Private/SAutomationGraphicalResultBox.h
+++ b/Engine/Source/Developer/AutomationWindow/Private/SAutomationGraphicalResultBox.h
@@ -30,6 +30,8 @@ public:
 
 	SLATE_END_ARGS()
 
+	virtual ~SAutomationGraphicalResultBox();
+
 public:
 
 	/**
diff --git a/Engine/Source/Developer/AutomationWindow/Private/SAutomationTestItem.cpp b/Engine/Source/Developer/AutomationWindow/Private/SAutomationTestItem.cpp
--- a/Engine/Source/Developer/AutomationWindow/Private/SAutomationTestItem.cpp
+++ b/Engine/Source/Developer/AutomationWindow/Private/SAutomationTestItem.cpp
@@ -1,7 +1,7 @@
 // Copyright 1998-2016 Epic Games, Inc. All Rights Reserved.
 
 #include "AutomationWindowPrivatePCH.h"
-
+#include "SHyperlink.h"
 
 #define LOCTEXT_NAMESPACE "AutomationTestItem"
 
@@ -139,6 +139,23 @@ TSharedRef<SWidget> SAutomationTestItem::GenerateWidgetForColumn( const FName& C
 
 	if( ColumnName == AutomationTestWindowConstants::Title)
 	{
+		TSharedRef<SWidget> TestNameWidget = SNullWidget::NullWidget;
+
+		// Would be nice to warp to text location...more difficult when distributed.
+		if ( TestStatus->GetSourceFile().IsEmpty() )
+		{
+			TestNameWidget = SNew(STextBlock)
+				.HighlightText(HighlightText)
+				.Text(FText::FromString(TestStatus->GetDisplayNameWithDecoration()));
+		}
+		else
+		{
+			TestNameWidget = SNew(SHyperlink)
+				.Style(FEditorStyle::Get(), "Common.GotoNativeCodeHyperlink")
+				.OnNavigate_Lambda([=] { FSlateApplication::Get().GotoLineInSource(TestStatus->GetSourceFile(), TestStatus->GetSourceFileLine()); })
+				.Text(FText::FromString(TestStatus->GetDisplayNameWithDecoration()));
+		}
+
 		return SNew(SHorizontalBox)
 			+SHorizontalBox::Slot()
 			.AutoWidth()
@@ -162,9 +179,7 @@ TSharedRef<SWidget> SAutomationTestItem::GenerateWidgetForColumn( const FName& C
 			.HAlign(HAlign_Center)
 			[
 				//name of the test
-				SNew( STextBlock )
-				.HighlightText( HighlightText )
-				.Text( FText::FromString(TestStatus->GetDisplayNameWithDecoration()) )
+				TestNameWidget
 			];
 	}
 	else if( ColumnName == AutomationTestWindowConstants::SmokeTest )
diff --git a/Engine/Source/Developer/AutomationWindow/Private/SAutomationWindow.cpp b/Engine/Source/Developer/AutomationWindow/Private/SAutomationWindow.cpp
--- a/Engine/Source/Developer/AutomationWindow/Private/SAutomationWindow.cpp
+++ b/Engine/Source/Developer/AutomationWindow/Private/SAutomationWindow.cpp
@@ -71,9 +71,10 @@ SAutomationWindow::~SAutomationWindow()
 	{
 		AutomationController->RemoveCallbacks();
 
-		AutomationController->OnControllerReset().Unbind();
-		AutomationController->OnTestsRefreshed().Unbind();
-		AutomationController->OnTestsAvailable().Unbind();
+		AutomationController->OnControllerReset().RemoveAll(this);
+		AutomationController->OnTestsRefreshed().RemoveAll(this);
+		AutomationController->OnTestsAvailable().RemoveAll(this);
+		AutomationController->OnTestsComplete().RemoveAll(this);
 	}
 }
 
@@ -92,9 +93,10 @@ void SAutomationWindow::Construct( const FArguments& InArgs, const IAutomationCo
 	SessionManager = InSessionManager;
 	AutomationController = InAutomationController;
 
-	AutomationController->OnControllerReset().BindRaw(this, &SAutomationWindow::OnRefreshTestCallback);
-	AutomationController->OnTestsRefreshed().BindRaw(this, &SAutomationWindow::OnRefreshTestCallback);
-	AutomationController->OnTestsAvailable().BindRaw(this, &SAutomationWindow::OnTestAvailableCallback);
+	AutomationController->OnControllerReset().AddRaw(this, &SAutomationWindow::OnRefreshTestCallback);
+	AutomationController->OnTestsRefreshed().AddRaw(this, &SAutomationWindow::OnRefreshTestCallback);
+	AutomationController->OnTestsAvailable().AddRaw(this, &SAutomationWindow::OnTestAvailableCallback);
+	AutomationController->OnTestsComplete().AddRaw(this, &SAutomationWindow::OnTestsCompleteCallback);
 
 	AutomationControllerState = AutomationController->GetTestState();
 	
@@ -124,6 +126,8 @@ void SAutomationWindow::Construct( const FArguments& InArgs, const IAutomationCo
 		.OnGenerateRow( this, &SAutomationWindow::OnGenerateWidgetForTest )
 		// Gets children
 		.OnGetChildren(this, &SAutomationWindow::OnGetChildren)
+		// on recursive expansion (shift + click)
+		.OnSetExpansionRecursive(this, &SAutomationWindow::OnTestExpansionRecursive)
 		//on selection
 		.OnSelectionChanged(this, &SAutomationWindow::OnTestSelectionChanged)
 		// Allow for some spacing between items with a larger item height.
@@ -136,7 +140,7 @@ void SAutomationWindow::Construct( const FArguments& InArgs, const IAutomationCo
 		(
 		SAssignNew(TestTableHeaderRow,SHeaderRow)
 		+ SHeaderRow::Column( AutomationTestWindowConstants::Title )
-		.FillWidth(300.0f)
+		.FillWidth(0.80f)
 		[
 			SNew(SHorizontalBox)
 			+SHorizontalBox::Slot()
@@ -156,12 +160,13 @@ void SAutomationWindow::Construct( const FArguments& InArgs, const IAutomationCo
 				.Text( LOCTEXT("TestName_Header", "Test Name") )
 			]
 		]
+
 		+ SHeaderRow::Column( AutomationTestWindowConstants::SmokeTest )
+		.FixedWidth( 50.0f )
 		.HAlignHeader(HAlign_Center)
 		.VAlignHeader(VAlign_Center)
 		.HAlignCell(HAlign_Center)
 		.VAlignCell(VAlign_Center)
-		.FixedWidth( 50.0f )
 		[
 			//icon for the smoke test column
 			SNew(SImage)
@@ -169,26 +174,29 @@ void SAutomationWindow::Construct( const FArguments& InArgs, const IAutomationCo
 			.ToolTipText( LOCTEXT( "Smoke Test", "Smoke Test" ) )
 			.Image(FEditorStyle::GetBrush("Automation.SmokeTest"))
 		]
+
 		+ SHeaderRow::Column( AutomationTestWindowConstants::RequiredDeviceCount )
-			.HAlignHeader(HAlign_Center)
-			.VAlignHeader(VAlign_Center)
-			.HAlignCell(HAlign_Center)
-			.VAlignCell(VAlign_Center)
-			.FixedWidth( 50.0f )
-			[
-				SNew( SImage )
-				.Image( FEditorStyle::GetBrush("Automation.ParticipantsWarning") )
-				.ToolTipText( LOCTEXT( "RequiredDeviceCountWarningToolTip", "Number of devices required." ) )
-			]
+		.FixedWidth(50.0f)
+		.HAlignHeader(HAlign_Center)
+		.VAlignHeader(VAlign_Center)
+		.HAlignCell(HAlign_Center)
+		.VAlignCell(VAlign_Center)
+		[
+			SNew( SImage )
+			.Image( FEditorStyle::GetBrush("Automation.ParticipantsWarning") )
+			.ToolTipText( LOCTEXT( "RequiredDeviceCountWarningToolTip", "Number of devices required." ) )
+		]
+
+		+ SHeaderRow::Column(AutomationTestWindowConstants::Timing)
+		.FixedWidth(100.0f)
+		.DefaultLabel(LOCTEXT("TestDurationRange", "Duration"))
+
 		+ SHeaderRow::Column( AutomationTestWindowConstants::Status )
-		.FillWidth(650.0f)
+		.FillWidth(0.20f)
 		[
 			//platform header placeholder
 			PlatformsHBox.ToSharedRef()
 		]
-		+ SHeaderRow::Column( AutomationTestWindowConstants::Timing )
-		.FillWidth(150.0f)
-		.DefaultLabel( LOCTEXT("TestDurationRange", "Duration") )
 
 		);
 
@@ -215,55 +223,102 @@ void SAutomationWindow::Construct( const FArguments& InArgs, const IAutomationCo
 		+ SOverlay::Slot()
 			[
 				SNew( SSplitter )
-					.IsEnabled(this, &SAutomationWindow::HandleMainContentIsEnabled)
-					.Orientation(Orient_Vertical)
+				.IsEnabled(this, &SAutomationWindow::HandleMainContentIsEnabled)
+				.Orientation(Orient_Vertical)
 
 				+ SSplitter::Slot()
-					.Value(0.66f)
-					[
-						//automation test panel
-						SAssignNew( MenuBar, SVerticalBox )
+				.Value(0.66f)
+				[
+					//automation test panel
+					SAssignNew( MenuBar, SVerticalBox )
 
-						//ACTIONS
-						+ SVerticalBox::Slot()
-							.AutoHeight()
-							[
-								SNew( SHorizontalBox )
+					//ACTIONS
+					+ SVerticalBox::Slot()
+					.AutoHeight()
+					[
+						SNew( SHorizontalBox )
 
-								+ SHorizontalBox::Slot()
-									.AutoWidth()
-									.HAlign(HAlign_Left)
-									.VAlign(VAlign_Center)
-									[
-										SAutomationWindow::MakeAutomationWindowToolBar( AutomationWindowActions.ToSharedRef(), SharedThis(this) )
-									]
-							]
+						+ SHorizontalBox::Slot()
+						.AutoWidth()
+						.HAlign(HAlign_Left)
+						.VAlign(VAlign_Center)
+						[
+							SAutomationWindow::MakeAutomationWindowToolBar( AutomationWindowActions.ToSharedRef(), SharedThis(this) )
+						]
+					]
 
-						+ SVerticalBox::Slot()
-							.FillHeight(1.0f)
-							.Padding(0.0f, 4.0f, 0.0f, 0.0f)
+					+ SVerticalBox::Slot()
+					.FillHeight(1.0f)
+					.Padding(0.0f, 4.0f, 0.0f, 0.0f)
+					[
+						SNew(SOverlay)
+						+ SOverlay::Slot()
+						[
+							SNew(SBorder)
+							.BorderImage(this, &SAutomationWindow::GetTestBackgroundBorderImage)
+							.Padding(3)
 							[
-								SNew(SOverlay)
-									+ SOverlay::Slot()
+								SNew(SBox)
+								.Padding(4)
+								[
+									SNew(SVerticalBox)
+
+									+ SVerticalBox::Slot()
+									.AutoHeight()
+									.Padding(0.0f, 0.0f, 0.0f, 4.0f)
 									[
-										SNew(SBorder)
-											.BorderImage(this, &SAutomationWindow::GetTestBackgroundBorderImage)
-											.Padding(3.0f)
+										SNew(SHorizontalBox)
+
+										+ SHorizontalBox::Slot()
+										.AutoWidth()
+										.VAlign(VAlign_Center)
+										[
+											SNew(SBox)
+											.MinDesiredWidth(130.0f)
 											[
-												//the actual table full of tests
-												TestTable.ToSharedRef()
+												SAssignNew(RequestedFilterComboBox, SComboBox< TSharedPtr<FString> >)
+												.OptionsSource(&RequestedFilterComboList)
+												.OnGenerateWidget(this, &SAutomationWindow::GenerateRequestedFilterComboItem)
+												.OnSelectionChanged(this, &SAutomationWindow::HandleRequesteFilterChanged)
+												.ContentPadding(FMargin(4.0, 1.0f))
+												[
+													SNew(STextBlock)
+													.Text(this, &SAutomationWindow::GetRequestedFilterComboText)
+												]
 											]
+										]
+
+										+ SHorizontalBox::Slot()
+										.FillWidth(1.0f)
+										.VAlign(VAlign_Center)
+										.Padding(2.0f, 0, 0, 0)
+										[
+											SAssignNew(AutomationSearchBox, SSearchBox)
+											.ToolTipText(LOCTEXT("Search Tests", "Search Tests"))
+											.OnTextChanged(this, &SAutomationWindow::OnFilterTextChanged)
+											.IsEnabled(this, &SAutomationWindow::IsAutomationControllerIdle)
+										]
 									]
 
-								+ SOverlay::Slot()
-									.HAlign(HAlign_Center)
-									.VAlign(VAlign_Center)
+									+ SVerticalBox::Slot()
+									.FillHeight(1.0f)
 									[
-										SNew(SThrobber)	
-											.Visibility(this, &SAutomationWindow::GetTestsUpdatingThrobberVisibility)
+										//the actual table full of tests
+										TestTable.ToSharedRef()
 									]
+								]
 							]
+						]
+
+						+ SOverlay::Slot()
+						.HAlign(HAlign_Center)
+						.VAlign(VAlign_Center)
+						[
+							SNew(SThrobber)	
+							.Visibility(this, &SAutomationWindow::GetTestsUpdatingThrobberVisibility)
+						]
 					]
+				]
 
 				+ SSplitter::Slot()
 					.Value(0.33f)
@@ -502,7 +557,7 @@ TSharedRef< SWidget > SAutomationWindow::MakeAutomationWindowToolBar( const TSha
 {
 	struct Local
 	{
-		static void FillToolbar(FToolBarBuilder& ToolbarBuilder, TSharedRef<SWidget> RunTests, TSharedRef<SWidget> Searchbox, TSharedRef<SWidget> PresetBox, TSharedRef<SWidget> HistoryBox, TSharedRef<SWidget> RequestedFilterBox, TWeakPtr<class SAutomationWindow> InAutomationWindow)
+		static void FillToolbar(FToolBarBuilder& ToolbarBuilder, TSharedRef<SWidget> RunTests, TSharedRef<SWidget> PresetBox, TSharedRef<SWidget> HistoryBox, TWeakPtr<class SAutomationWindow> InAutomationWindow)
 		{
 			ToolbarBuilder.BeginSection("Automation");
 			{
@@ -524,7 +579,6 @@ TSharedRef< SWidget > SAutomationWindow::MakeAutomationWindowToolBar( const TSha
 				ToolbarBuilder.AddToolBarButton( FAutomationWindowCommands::Get().ErrorFilter );
 				ToolbarBuilder.AddToolBarButton( FAutomationWindowCommands::Get().WarningFilter );
 				ToolbarBuilder.AddToolBarButton( FAutomationWindowCommands::Get().DeveloperDirectoryContent );
-				ToolbarBuilder.AddWidget(RequestedFilterBox);
 			}
 			ToolbarBuilder.EndSection();
 			ToolbarBuilder.BeginSection("GroupFlags");
@@ -537,11 +591,6 @@ TSharedRef< SWidget > SAutomationWindow::MakeAutomationWindowToolBar( const TSha
 					FSlateIcon(FEditorStyle::GetStyleSetName(), "AutomationWindow.GroupSettings"));
 			}
 			ToolbarBuilder.EndSection();
-			ToolbarBuilder.BeginSection("Search");
-			{
-				ToolbarBuilder.AddWidget( Searchbox );
-			}
-			ToolbarBuilder.EndSection();
 			ToolbarBuilder.BeginSection("Presets");
 			{
 				ToolbarBuilder.AddWidget( PresetBox );
@@ -613,21 +662,6 @@ TSharedRef< SWidget > SAutomationWindow::MakeAutomationWindowToolBar( const TSha
 			]
 		];
 
-	const float SearchWidth = 200.f;
-	TSharedRef<SWidget> Searchbox = 
-		SNew(SHorizontalBox)
-		+SHorizontalBox::Slot()
-		.MaxWidth(SearchWidth)
-		.FillWidth(1.f)
-		.VAlign(VAlign_Bottom)
-		[
-			SAssignNew( AutomationSearchBox, SSearchBox )
-			.ToolTipText( LOCTEXT( "Search Tests", "Search Tests" ) )
-			.OnTextChanged( this, &SAutomationWindow::OnFilterTextChanged )
-			.IsEnabled( this, &SAutomationWindow::IsAutomationControllerIdle )
-			.MinDesiredWidth(SearchWidth)
-		];
-
 	const float HistoryWidth = 200.0f;
 	TSharedRef<SWidget> History =
 		SNew(SHorizontalBox)
@@ -777,26 +811,17 @@ TSharedRef< SWidget > SAutomationWindow::MakeAutomationWindowToolBar( const TSha
 		];
 
 	RequestedFilterComboList.Empty();
+	RequestedFilterComboList.Add(MakeShareable(new FString(TEXT("All Tests"))));
 	RequestedFilterComboList.Add(MakeShareable(new FString(TEXT("Smoke Tests"))));
 	RequestedFilterComboList.Add(MakeShareable(new FString(TEXT("Engine Tests"))));
 	RequestedFilterComboList.Add(MakeShareable(new FString(TEXT("Product Tests"))));
-	RequestedFilterComboList.Add(MakeShareable(new FString(TEXT("Perf Tests"))));
+	RequestedFilterComboList.Add(MakeShareable(new FString(TEXT("Performance Tests"))));
 	RequestedFilterComboList.Add(MakeShareable(new FString(TEXT("Stress Tests"))));
 	RequestedFilterComboList.Add(MakeShareable(new FString(TEXT("Standard Tests"))));
-	RequestedFilterComboList.Add(MakeShareable(new FString(TEXT("All Tests"))));
-
-	SAssignNew(RequestedFilterComboBox, SComboBox< TSharedPtr<FString> >)
-		.OptionsSource(&RequestedFilterComboList)
-		.OnGenerateWidget(this, &SAutomationWindow::GenerateRequestedFilterComboItem)
-		.OnSelectionChanged(this, &SAutomationWindow::HandleRequesteFilterChanged)
-		[
-			SNew(STextBlock)
-			.Text(this, &SAutomationWindow::GetRequestedFilterComboText)
-		];
 
 	FToolBarBuilder ToolbarBuilder( InCommandList, FMultiBoxCustomization::None );
 	TWeakPtr<SAutomationWindow> AutomationWindow = SharedThis(this);
-	Local::FillToolbar(ToolbarBuilder, RunTests, Searchbox, TestPresets, History, RequestedFilterComboBox.ToSharedRef(), AutomationWindow);
+	Local::FillToolbar(ToolbarBuilder, RunTests, TestPresets, History, AutomationWindow);
 
 	// Create the tool bar!
 	return
@@ -895,27 +920,27 @@ void SAutomationWindow::HandleRequesteFilterChanged(TSharedPtr<FString> Item, ES
 
 	switch (EntryIndex)
 	{
-		case 0:	//	"Smoke Tests"
+		case 0:	//	"All Tests"
+			NewRequestedFlags = EAutomationTestFlags::FilterMask;
+			break;
+		case 1:	//	"Smoke Tests"
 			NewRequestedFlags = EAutomationTestFlags::SmokeFilter;
 			break;
-		case 1:	//	"Engine Tests"
+		case 2:	//	"Engine Tests"
 			NewRequestedFlags = EAutomationTestFlags::EngineFilter;
 			break;
-		case 2:	//	"Product Tests"
+		case 3:	//	"Product Tests"
 			NewRequestedFlags = EAutomationTestFlags::ProductFilter;
 			break;
-		case 3:	//	"Perf Tests"
+		case 4:	//	"Performance Tests"
 			NewRequestedFlags = EAutomationTestFlags::PerfFilter;
 			break;
-		case 4:	//	"Stress Tests"
+		case 5:	//	"Stress Tests"
 			NewRequestedFlags = EAutomationTestFlags::StressFilter;
 			break;
-		case 5:	//	"Standard Tests"
+		case 6:	//	"Standard Tests"
 			NewRequestedFlags = EAutomationTestFlags::SmokeFilter | EAutomationTestFlags::EngineFilter | EAutomationTestFlags::ProductFilter | EAutomationTestFlags::PerfFilter;
 			break;
-		case 6:	//	"All Tests"
-			NewRequestedFlags = EAutomationTestFlags::FilterMask;
-			break;
 	}
 	AutomationController->SetRequestedTestFlags(NewRequestedFlags);
 }
@@ -984,7 +1009,7 @@ FText SAutomationWindow::GetRequestedFilterComboText() const
 	}
 	else
 	{
-		return LOCTEXT("AutomationRequestedFilterComboLabel", "<Select A Filter>");
+		return LOCTEXT("AutomationRequestedFilterComboLabel", "All Tests");
 	}
 }
 
@@ -1270,7 +1295,7 @@ bool SAutomationWindow::IsFullSizeScreenshotsOptionEnabled() const
 #if WITH_EDITOR
 TSharedPtr<SWidget> SAutomationWindow::HandleAutomationListContextMenuOpening()
 {
- 	TArray< TSharedPtr<IAutomationReport> >SelectedReport = TestTable->GetSelectedItems();
+	TArray< TSharedPtr<IAutomationReport> >SelectedReport = TestTable->GetSelectedItems();
 
 	TArray<FString> AssetNames;
 	for (int32 ReportIndex = 0; ReportIndex < SelectedReport.Num(); ++ReportIndex)
@@ -1306,6 +1331,20 @@ void SAutomationWindow::OnGetChildren(TSharedPtr<IAutomationReport> InItem, TArr
 	OutItems = InItem->GetFilteredChildren();
 }
 
+void SAutomationWindow::OnTestExpansionRecursive(TSharedPtr<IAutomationReport> InAutomationReport, bool bInIsItemExpanded)
+{
+	if ( InAutomationReport.IsValid() )
+	{
+		TArray<TSharedPtr<IAutomationReport> >& FilteredChildren = InAutomationReport->GetFilteredChildren();
+
+		TestTable->SetItemExpansion(InAutomationReport, bInIsItemExpanded);
+
+		for ( TSharedPtr<IAutomationReport>& Child : FilteredChildren )
+		{
+			OnTestExpansionRecursive(Child, bInIsItemExpanded);
+		}
+	}
+}
 
 void SAutomationWindow::OnTestSelectionChanged(TSharedPtr<IAutomationReport> Selection, ESelectInfo::Type /*SelectInfo*/)
 {
@@ -1568,6 +1607,15 @@ void SAutomationWindow::OnTestAvailableCallback( EAutomationControllerModuleStat
 	AutomationControllerState = InAutomationControllerState;
 }
 
+void SAutomationWindow::OnTestsCompleteCallback()
+{
+	// Simulate selection again after testing finishes.
+	if ( TestTable->GetNumItemsSelected() > 0 )
+	{
+		OnTestSelectionChanged(TestTable->GetSelectedItems()[0], ESelectInfo::Direct);
+	}
+}
+
 
 void SAutomationWindow::ExpandTreeView( TSharedPtr< IAutomationReport > InReport, const bool ShouldExpand )
 {
diff --git a/Engine/Source/Developer/AutomationWindow/Private/SAutomationWindow.h b/Engine/Source/Developer/AutomationWindow/Private/SAutomationWindow.h
--- a/Engine/Source/Developer/AutomationWindow/Private/SAutomationWindow.h
+++ b/Engine/Source/Developer/AutomationWindow/Private/SAutomationWindow.h
@@ -200,6 +200,9 @@ private:
 
 	/** Gets children tests for a node in the hierarchy */
 	void OnGetChildren(TSharedPtr<IAutomationReport> InItem, TArray<TSharedPtr<IAutomationReport> >& OutItems);
+
+	/** Callback for a test expansion changing recursively */
+	void OnTestExpansionRecursive(TSharedPtr<IAutomationReport> InTreeNode, bool bInIsItemExpanded);
 	
 	/** Callback for a new test being selected */
 	void OnTestSelectionChanged(TSharedPtr<IAutomationReport> Selection, ESelectInfo::Type SelectInfo);
@@ -383,12 +386,17 @@ private:
 	bool IsAutomationRunButtonEnabled() const;
 
 	/**
- 	 * Set whether tests are available to run.
+	 * Set whether tests are available to run.
 	 *
 	 * @param The Automation controller state.
 	 */
 	void OnTestAvailableCallback( EAutomationControllerModuleState::Type InAutomationControllerState );
 
+	/**
+	 * Called when tests complete.
+	 */
+	void OnTestsCompleteCallback();
+
 	/** Copies the selected log messages to the clipboard. */
 	void CopyLog();
 
diff --git a/Engine/Source/Developer/BlueprintProfiler/BlueprintProfiler.Build.cs b/Engine/Source/Developer/BlueprintProfiler/BlueprintProfiler.Build.cs
--- a/Engine/Source/Developer/BlueprintProfiler/BlueprintProfiler.Build.cs
+++ b/Engine/Source/Developer/BlueprintProfiler/BlueprintProfiler.Build.cs
@@ -25,6 +25,7 @@ public class BlueprintProfiler : ModuleRules
             PrivateDependencyModuleNames.Add("BlueprintGraph");
             PrivateDependencyModuleNames.Add("EditorStyle");
             PrivateDependencyModuleNames.Add("InputCore");
+            PrivateDependencyModuleNames.Add("SlateCore");
         }
     }
 }
diff --git a/Engine/Source/Developer/BlueprintProfiler/Private/ScriptInstrumentationPlayback.cpp b/Engine/Source/Developer/BlueprintProfiler/Private/ScriptInstrumentationPlayback.cpp
--- a/Engine/Source/Developer/BlueprintProfiler/Private/ScriptInstrumentationPlayback.cpp
+++ b/Engine/Source/Developer/BlueprintProfiler/Private/ScriptInstrumentationPlayback.cpp
@@ -271,7 +271,7 @@ void FBlueprintFunctionContext::InitialiseContextFromGraph(TSharedPtr<FBlueprint
 					EventParams.Tooltip = LOCTEXT("NavigateToEventLocationHyperlink_ToolTip", "Navigate to the Event");
 					EventParams.NodeFlags = EScriptExecutionNodeFlags::Event;
 					EventParams.IconColor = FLinearColor(0.91f, 0.16f, 0.16f);
-					const FSlateBrush* EventIcon = EventNode->ShowPaletteIconOnNode() ?	FEditorStyle::GetBrush(EventNode->GetPaletteIcon(EventParams.IconColor)) :
+					const FSlateBrush* EventIcon = EventNode->ShowPaletteIconOnNode() ?	EventNode->GetIconAndTint(EventParams.IconColor).GetOptionalIcon() :
 																						FEditorStyle::GetBrush(TEXT("BlueprintProfiler.BPNode"));
 					EventParams.Icon = const_cast<FSlateBrush*>(EventIcon);
 					TSharedPtr<FScriptExecutionNode> EventExecNode = CreateExecutionNode(EventParams);
@@ -304,7 +304,7 @@ void FBlueprintFunctionContext::InitialiseContextFromGraph(TSharedPtr<FBlueprint
 					FunctionNodeParams.Tooltip = LOCTEXT("NavigateToFunctionLocationHyperlink_ToolTip", "Navigate to the Function");
 					FunctionNodeParams.NodeFlags = EScriptExecutionNodeFlags::Event;
 					FunctionNodeParams.IconColor = FLinearColor(0.91f, 0.16f, 0.16f);
-					const FSlateBrush* Icon = FunctionEntry->ShowPaletteIconOnNode() ?	FEditorStyle::GetBrush(FunctionEntry->GetPaletteIcon(FunctionNodeParams.IconColor)) :
+					const FSlateBrush* Icon = FunctionEntry->ShowPaletteIconOnNode() ? FunctionEntry->GetIconAndTint(FunctionNodeParams.IconColor).GetOptionalIcon() :
 																						FEditorStyle::GetBrush(TEXT("BlueprintProfiler.BPNode"));
 					FunctionNodeParams.Icon = const_cast<FSlateBrush*>(Icon);
 					TSharedPtr<FScriptExecutionNode> FunctionEntryNode = CreateExecutionNode(FunctionNodeParams);
@@ -502,7 +502,7 @@ void FBlueprintFunctionContext::CreateInputEvents(TSharedPtr<FBlueprintExecution
 		EventParams.Tooltip = LOCTEXT("NavigateToEventLocationHyperlink_ToolTip", "Navigate to the Event");
 		EventParams.NodeFlags = EScriptExecutionNodeFlags::Event;
 		EventParams.IconColor = FLinearColor(0.91f, 0.16f, 0.16f);
-		const FSlateBrush* EventIcon = GraphNode->ShowPaletteIconOnNode() ?	FEditorStyle::GetBrush(GraphNode->GetPaletteIcon(EventParams.IconColor)) :
+		const FSlateBrush* EventIcon = GraphNode->ShowPaletteIconOnNode() ? GraphNode->GetIconAndTint(EventParams.IconColor).GetOptionalIcon() :
 																			FEditorStyle::GetBrush(TEXT("BlueprintProfiler.BPNode"));
 		EventParams.Icon = const_cast<FSlateBrush*>(EventIcon);
 		TSharedPtr<FScriptExecutionNode> EventExecNode = CreateExecutionNode(EventParams);
@@ -526,7 +526,7 @@ TSharedPtr<FScriptExecutionNode> FBlueprintFunctionContext::MapNodeExecution(UEd
 			NodeParams.Tooltip = LOCTEXT("NavigateToNodeLocationHyperlink_ToolTip", "Navigate to the Node");
 			NodeParams.NodeFlags = EScriptExecutionNodeFlags::Node;
 			NodeParams.IconColor = FLinearColor(1.f, 1.f, 1.f, 0.8f);
-			const FSlateBrush* NodeIcon = NodeToMap->ShowPaletteIconOnNode() ?	FEditorStyle::GetBrush(NodeToMap->GetPaletteIcon(NodeParams.IconColor)) :
+			const FSlateBrush* NodeIcon = NodeToMap->ShowPaletteIconOnNode() ? NodeToMap->GetIconAndTint(NodeParams.IconColor).GetOptionalIcon() :
 																				FEditorStyle::GetBrush(TEXT("BlueprintProfiler.BPNode"));
 			NodeParams.Icon = const_cast<FSlateBrush*>(NodeIcon);
 			MappedNode = CreateExecutionNode(NodeParams);
@@ -579,7 +579,7 @@ TSharedPtr<FScriptExecutionNode> FBlueprintFunctionContext::MapNodeExecution(UEd
 					const bool bEventCall = FunctionNode ? FunctionNode->IsA<UK2Node_Event>() : false;
 					MappedNode->AddFlags(EScriptExecutionNodeFlags::FunctionCall);
 					MappedNode->SetIconColor(FLinearColor(0.46f, 0.54f, 0.95f));
-					MappedNode->SetToolTipText(LOCTEXT("NavigateToFunctionLocationHyperlink_ToolTip", "Navigate to the Function Callsite"));
+					MappedNode->SetToolTipText(LOCTEXT("NavigateToFunctionCallsiteHyperlink_ToolTip", "Navigate to the Function Callsite"));
 					// Don't add entry points for events.
 					if (!bEventCall)
 					{
@@ -670,7 +670,7 @@ void FBlueprintFunctionContext::MapExecPins(TSharedPtr<FScriptExecutionNode> Exe
 			LinkNodeParams.NodeName = Pin->GetFName();
 			LinkNodeParams.ObservedObject = Pin;
 			LinkNodeParams.DisplayName = Pin->GetDisplayName();
-			LinkNodeParams.Tooltip = LOCTEXT("ExecPin_ToolTip", "Expand execution path");
+			LinkNodeParams.Tooltip = LOCTEXT("ExecPin_ExpandExecutionPath_ToolTip", "Expand execution path");
 			LinkNodeParams.NodeFlags = !bInvalidTrace ? EScriptExecutionNodeFlags::ExecPin : (EScriptExecutionNodeFlags::ExecPin|EScriptExecutionNodeFlags::InvalidTrace);
 			LinkNodeParams.IconColor = FLinearColor(1.f, 1.f, 1.f, 0.8f);
 			const FSlateBrush* Icon = LinkedPin ?	FEditorStyle::GetBrush(TEXT("BlueprintProfiler.BPPinConnected")) : 
@@ -774,8 +774,8 @@ TSharedPtr<FScriptExecutionNode> FBlueprintFunctionContext::MapTunnelEntry(const
 		TunnelInstanceParams.Tooltip = LOCTEXT("NavigateToNodeLocationHyperlink_ToolTip", "Navigate to the Node");
 		TunnelInstanceParams.NodeFlags = EScriptExecutionNodeFlags::FunctionTunnel;
 		TunnelInstanceParams.IconColor = FLinearColor(1.f, 1.f, 1.f, 0.8f);
-		const FSlateBrush* TunnelIcon = TunnelNode->ShowPaletteIconOnNode() ?	FEditorStyle::GetBrush(TunnelNode->GetPaletteIcon(TunnelInstanceParams.IconColor)) :
-																				FEditorStyle::GetBrush(TEXT("BlueprintProfiler.BPNode"));
+		const FSlateBrush* TunnelIcon = TunnelNode->ShowPaletteIconOnNode() ? TunnelNode->GetIconAndTint(TunnelInstanceParams.IconColor).GetOptionalIcon() :
+																			  FEditorStyle::GetBrush(TEXT("BlueprintProfiler.BPNode"));
 		TunnelInstanceParams.Icon = const_cast<FSlateBrush*>(TunnelIcon);
 		TunnelInstance = CreateTypedExecutionNode<FScriptExecutionTunnelInstance>(TunnelInstanceParams);
 	}
@@ -787,8 +787,8 @@ TSharedPtr<FScriptExecutionNode> FBlueprintFunctionContext::MapTunnelEntry(const
 	TunnelParams.Tooltip = LOCTEXT("NavigateToNodeLocationHyperlink_ToolTip", "Navigate to the Node");
 	TunnelParams.NodeFlags = EScriptExecutionNodeFlags::FunctionTunnel;
 	TunnelParams.IconColor = FLinearColor(1.f, 1.f, 1.f, 0.8f);
-	const FSlateBrush* TunnelIcon = TunnelNode->ShowPaletteIconOnNode() ?	FEditorStyle::GetBrush(TunnelNode->GetPaletteIcon(TunnelParams.IconColor)) :
-																			FEditorStyle::GetBrush(TEXT("BlueprintProfiler.BPNode"));
+	const FSlateBrush* TunnelIcon = TunnelNode->ShowPaletteIconOnNode() ? TunnelNode->GetIconAndTint(TunnelParams.IconColor).GetOptionalIcon() :
+																		  FEditorStyle::GetBrush(TEXT("BlueprintProfiler.BPNode"));
 	TunnelParams.Icon = const_cast<FSlateBrush*>(TunnelIcon);
 	TSharedPtr<FScriptExecutionTunnelEntry> TunnelCallNode = MakeShareable(new FScriptExecutionTunnelEntry(TunnelParams, TunnelNode));
 	ExecutionNodes.FindOrAdd(TunnelParams.NodeName) = TunnelCallNode;
@@ -1117,7 +1117,7 @@ void FBlueprintFunctionContext::MapTunnelIO(UEdGraph* TunnelGraph)
 					LinkParams.NodeName = FName(*Pin->PinName);
 					LinkParams.ObservedObject = Pin;
 					LinkParams.DisplayName = Pin->GetDisplayName().IsEmpty() ? FText::FromString(*Pin->PinName) : Pin->GetDisplayName();
-					LinkParams.Tooltip = LOCTEXT("ExecPin_ToolTip", "Expand tunnel entry point");
+					LinkParams.Tooltip = LOCTEXT("ExecPin_ExpandTunnelEntryPoint_ToolTip", "Expand tunnel entry point");
 					LinkParams.NodeFlags = EScriptExecutionNodeFlags::TunnelEntryPin;
 					LinkParams.IconColor = FLinearColor(1.f, 1.f, 1.f, 0.8f);
 					const FSlateBrush* Icon = FEditorStyle::GetBrush(TEXT("BlueprintProfiler.BPPinConnected"));
@@ -1132,7 +1132,7 @@ void FBlueprintFunctionContext::MapTunnelIO(UEdGraph* TunnelGraph)
 					LinkParams.NodeName = FName(*Pin->PinName);
 					LinkParams.ObservedObject = Pin;
 					LinkParams.DisplayName = Pin->GetDisplayName().IsEmpty() ? FText::FromString(*Pin->PinName) : Pin->GetDisplayName();
-					LinkParams.Tooltip = LOCTEXT("ExecPin_ToolTip", "Expand tunnel exit point");
+					LinkParams.Tooltip = LOCTEXT("ExecPin_ExpandTunnelExitPoint_ToolTip", "Expand tunnel exit point");
 					LinkParams.NodeFlags = EScriptExecutionNodeFlags::TunnelExitPin;
 					LinkParams.IconColor = FLinearColor(1.f, 1.f, 1.f, 0.8f);
 					const FSlateBrush* Icon = FEditorStyle::GetBrush(TEXT("BlueprintProfiler.BPPinConnected"));
diff --git a/Engine/Source/Developer/DeviceManager/Private/Models/DeviceDetailsCommands.h b/Engine/Source/Developer/DeviceManager/Private/Models/DeviceDetailsCommands.h
--- a/Engine/Source/Developer/DeviceManager/Private/Models/DeviceDetailsCommands.h
+++ b/Engine/Source/Developer/DeviceManager/Private/Models/DeviceDetailsCommands.h
@@ -2,7 +2,7 @@
 
 #pragma once
 
-#define LOCTEXT_NAMESPACE ""
+#define LOCTEXT_NAMESPACE "DeviceDetailsCommands"
 
 /**
  * The device details commands
diff --git a/Engine/Source/Developer/DirectoryWatcher/Private/Windows/DirectoryWatchRequestWindows.cpp b/Engine/Source/Developer/DirectoryWatcher/Private/Windows/DirectoryWatchRequestWindows.cpp
--- a/Engine/Source/Developer/DirectoryWatcher/Private/Windows/DirectoryWatchRequestWindows.cpp
+++ b/Engine/Source/Developer/DirectoryWatcher/Private/Windows/DirectoryWatchRequestWindows.cpp
@@ -95,6 +95,7 @@ bool FDirectoryWatchRequestWindows::Init(const FString& InDirectory)
 	if ( !bSuccess  )
 	{
 		::CloseHandle(DirectoryHandle);
+		DirectoryHandle = INVALID_HANDLE_VALUE;
 		return false;
 	}
 
@@ -173,7 +174,8 @@ void FDirectoryWatchRequestWindows::ProcessChange(uint32 Error, uint32 NumBytes)
 	{
 		// The operation was aborted, likely due to EndWatchRequest canceling it.
 		// Mark the request for delete so it can be cleaned up next tick.
-		bPendingDelete = true; 
+		bPendingDelete = true;
+		UE_LOG(LogDirectoryWatcher, Warning, TEXT("A directory notification for '%s' was aborted."), *Directory);
 		return; 
 	}
 
@@ -186,6 +188,11 @@ void FDirectoryWatchRequestWindows::ProcessChange(uint32 Error, uint32 NumBytes)
 		FMemory::Memcpy(BackBuffer, Buffer, NumBytes);
 	}
 
+	if ( !bValidNotification )
+	{
+		UE_LOG(LogDirectoryWatcher, Warning, TEXT("A directory notification failed for '%s' because it was empty or there was a buffer overflow. Attemping another request..."), *Directory);
+	}
+
 	// Start up another read
 	const bool bSuccess = !!::ReadDirectoryChangesW(
 		DirectoryHandle,
@@ -197,9 +204,14 @@ void FDirectoryWatchRequestWindows::ProcessChange(uint32 Error, uint32 NumBytes)
 		&Overlapped,
 		&FDirectoryWatchRequestWindows::ChangeNotification);
 
-	if ( !bValidNotification )
+	if ( !bSuccess  )
 	{
-		UE_LOG(LogDirectoryWatcher, Warning, TEXT("A directory notification failed for '%s' because it was aborted or there was a buffer overflow."), *Directory);
+		// Failed to re-create the read request.
+		// Mark the request for delete so it can be cleaned up next tick.
+		::CloseHandle(DirectoryHandle);
+		DirectoryHandle = INVALID_HANDLE_VALUE;
+		bPendingDelete = true;
+		UE_LOG(LogDirectoryWatcher, Warning, TEXT("A directory notification failed for '%s', and we were unable to create a new request."), *Directory);
 		return;
 	}
 
diff --git a/Engine/Source/Developer/HotReload/Private/HotReload.cpp b/Engine/Source/Developer/HotReload/Private/HotReload.cpp
--- a/Engine/Source/Developer/HotReload/Private/HotReload.cpp
+++ b/Engine/Source/Developer/HotReload/Private/HotReload.cpp
@@ -559,7 +559,7 @@ bool FHotReloadModule::RecompileModule(const FName InModuleName, const bool bRel
 
 	ModuleCompilerStartedEvent.Broadcast(false); // we never perform an async compile
 
-	FModuleManager ModuleManager = FModuleManager::Get();
+	FModuleManager& ModuleManager = FModuleManager::Get();
 
 	// Update our set of known modules, in case we don't already know about this module
 	ModuleManager.AddModule( InModuleName );
diff --git a/Engine/Source/Developer/OutputLog/Private/SOutputLog.cpp b/Engine/Source/Developer/OutputLog/Private/SOutputLog.cpp
--- a/Engine/Source/Developer/OutputLog/Private/SOutputLog.cpp
+++ b/Engine/Source/Developer/OutputLog/Private/SOutputLog.cpp
@@ -1149,3 +1149,5 @@ bool FLogFilter::IsMessageAllowed(const TSharedPtr<FLogMessage>& Message)
 
 	return true;
 }
+
+#undef LOCTEXT_NAMESPACE
diff --git a/Engine/Source/Developer/ProjectLauncher/Private/Models/ProjectLauncherCommands.h b/Engine/Source/Developer/ProjectLauncher/Private/Models/ProjectLauncherCommands.h
--- a/Engine/Source/Developer/ProjectLauncher/Private/Models/ProjectLauncherCommands.h
+++ b/Engine/Source/Developer/ProjectLauncher/Private/Models/ProjectLauncherCommands.h
@@ -2,7 +2,7 @@
 
 #pragma once
 
-#define LOCTEXT_NAMESPACE ""
+#define LOCTEXT_NAMESPACE "ProjectLauncherCommands"
 
 /**
  * The project launcher UI commands
diff --git a/Engine/Source/Developer/ProjectLauncher/Private/Widgets/Package/SProjectLauncherPackagePage.cpp b/Engine/Source/Developer/ProjectLauncher/Private/Widgets/Package/SProjectLauncherPackagePage.cpp
--- a/Engine/Source/Developer/ProjectLauncher/Private/Widgets/Package/SProjectLauncherPackagePage.cpp
+++ b/Engine/Source/Developer/ProjectLauncher/Private/Widgets/Package/SProjectLauncherPackagePage.cpp
@@ -78,7 +78,7 @@ void SProjectLauncherPackagePage::Construct( const FArguments& InArgs, const FPr
 			.FillHeight(1.0)
 			.Padding(0.0, 8.0, 0.0, 0.0)
 			[
-				SNew(SProjectLauncherPackagingSettings, InModel)
+				SAssignNew(ProjectLauncherPackagingSettings, SProjectLauncherPackagingSettings, InModel)
 					.Visibility(this, &SProjectLauncherPackagePage::HandlePackagingSettingsAreaVisibility)
 			]
 	];
@@ -127,6 +127,9 @@ void SProjectLauncherPackagePage::HandlePackagingModeMenuEntryClicked( ELauncher
 	if (SelectedProfile.IsValid())
 	{
 		SelectedProfile->SetPackagingMode(PackagingMode);
+
+		check(ProjectLauncherPackagingSettings.IsValid());
+		ProjectLauncherPackagingSettings->UpdateDirectoryPathText();
 	}
 }
 
diff --git a/Engine/Source/Developer/ProjectLauncher/Private/Widgets/Package/SProjectLauncherPackagePage.h b/Engine/Source/Developer/ProjectLauncher/Private/Widgets/Package/SProjectLauncherPackagePage.h
--- a/Engine/Source/Developer/ProjectLauncher/Private/Widgets/Package/SProjectLauncherPackagePage.h
+++ b/Engine/Source/Developer/ProjectLauncher/Private/Widgets/Package/SProjectLauncherPackagePage.h
@@ -49,4 +49,6 @@ private:
 
 	// Holds a pointer to the data model.
 	FProjectLauncherModelPtr Model;
+
+	TSharedPtr<SProjectLauncherPackagingSettings> ProjectLauncherPackagingSettings;
 };
diff --git a/Engine/Source/Developer/ProjectLauncher/Private/Widgets/Package/SProjectLauncherPackagingSettings.cpp b/Engine/Source/Developer/ProjectLauncher/Private/Widgets/Package/SProjectLauncherPackagingSettings.cpp
--- a/Engine/Source/Developer/ProjectLauncher/Private/Widgets/Package/SProjectLauncherPackagingSettings.cpp
+++ b/Engine/Source/Developer/ProjectLauncher/Private/Widgets/Package/SProjectLauncherPackagingSettings.cpp
@@ -44,7 +44,7 @@ void SProjectLauncherPackagingSettings::Construct( const FArguments& InArgs, con
 							.AutoHeight()
 							[
 								SNew(STextBlock)
-									.Text(LOCTEXT("RepositoryPathLabel", "Repository Path:"))
+									.Text(this, &SProjectLauncherPackagingSettings::HandleDirectoryTitleText)
 							]
 
 						+ SVerticalBox::Slot()
@@ -58,7 +58,10 @@ void SProjectLauncherPackagingSettings::Construct( const FArguments& InArgs, con
 									.Padding(0.0, 0.0, 0.0, 3.0)
 									[
 										// repository path text box
-										SAssignNew(RepositoryPathTextBox, SEditableTextBox)
+										SAssignNew(DirectoryPathTextBox, SEditableTextBox)
+										.IsEnabled(this, &SProjectLauncherPackagingSettings::IsEditable)
+										.OnTextCommitted(this, &SProjectLauncherPackagingSettings::OnTextCommitted)
+										.OnTextChanged(this, &SProjectLauncherPackagingSettings::OnTextChanged)
 									]
 
 								+ SHorizontalBox::Slot()
@@ -69,62 +72,12 @@ void SProjectLauncherPackagingSettings::Construct( const FArguments& InArgs, con
 										// browse button
 										SNew(SButton)
 											.ContentPadding(FMargin(6.0, 2.0))
-											.IsEnabled(false)
+											.IsEnabled(this, &SProjectLauncherPackagingSettings::IsEditable)
 											.Text(LOCTEXT("BrowseButtonText", "Browse..."))
-											.ToolTipText(LOCTEXT("BrowseButtonToolTip", "Browse for the repository"))
-											//.OnClicked(this, &SProjectLauncherPackagePage::HandleBrowseButtonClicked)
+											.ToolTipText(LOCTEXT("BrowseButtonToolTip", "Browse for the directory"))
+											.OnClicked(this, &SProjectLauncherPackagingSettings::HandleBrowseButtonClicked)
 									]
 							]
-
-						+ SVerticalBox::Slot()
-							.AutoHeight()
-							.Padding(0.0, 8.0, 0.0, 0.0)
-							[
-								SNew(STextBlock)
-									.Text(LOCTEXT("DescriptionTextBoxLabel", "Description:"))
-							]
-
-						+ SVerticalBox::Slot()
-							.FillHeight(1.0)
-							.Padding(0.0, 4.0, 0.0, 0.0)
-							[
-								SNew(SEditableTextBox)
-							]
-					]
-			]
-
-		+ SVerticalBox::Slot()
-			.AutoHeight()
-			.Padding(0.0, 8.0, 0.0, 0.0)
-			[
-				SNew(SExpandableArea)
-					.AreaTitle(LOCTEXT("IncrementalPackagingAreaTitle", "Incremental Packaging"))
-					.InitiallyCollapsed(true)
-					.Padding(8.0)
-					.BodyContent()
-					[
-						SNew(SVerticalBox)
-
-						+ SVerticalBox::Slot()
-							.AutoHeight()
-							[
-								// incremental packaging check box
-								SNew(SCheckBox)
-									//.IsChecked(this, &SProjectLauncherLaunchPage::HandleIncrementalCheckBoxIsChecked)
-									//.OnCheckStateChanged(this, &SProjectLauncherLaunchPage::HandleIncrementalCheckBoxCheckStateChanged)
-									.Content()
-									[
-										SNew(STextBlock)
-											.Text(LOCTEXT("IncrementalCheckBoxText", "This is an incremental build"))
-									]
-							]
-
-						+ SVerticalBox::Slot()
-							.FillHeight(1.0)
-							.Padding(18.0, 12.0, 0.0, 0.0)
-							[
-								SNew(SVerticalBox)
-							]
 					]
 			]
 	];
@@ -133,14 +86,122 @@ void SProjectLauncherPackagingSettings::Construct( const FArguments& InArgs, con
 }
 END_SLATE_FUNCTION_BUILD_OPTIMIZATION
 
+void SProjectLauncherPackagingSettings::UpdateDirectoryPathText()
+{
+	DirectoryPathTextBox->SetText(HandleDirectoryPathText());
+}
 
 /* SProjectLauncherPackagingSettings callbacks
  *****************************************************************************/
 
+FText SProjectLauncherPackagingSettings::HandleDirectoryTitleText() const
+{
+	ILauncherProfilePtr SelectedProfile = Model->GetSelectedProfile();
+
+	if (SelectedProfile.IsValid())
+	{
+		ELauncherProfilePackagingModes::Type PackagingMode = SelectedProfile->GetPackagingMode();
+		switch (PackagingMode)
+		{
+		case ELauncherProfilePackagingModes::Locally:
+			return LOCTEXT("LocalDirectoryPathLabel", "Local Directory Path:");
+		case ELauncherProfilePackagingModes::SharedRepository:
+			return LOCTEXT("RepositoryPathLabel", "Repository Path:");
+		}
+	}
+
+	return FText::GetEmpty();
+}
+
+FText SProjectLauncherPackagingSettings::HandleDirectoryPathText() const
+{
+	ILauncherProfilePtr SelectedProfile = Model->GetSelectedProfile();
+
+	if (SelectedProfile.IsValid())
+	{
+		ELauncherProfilePackagingModes::Type PackagingMode = SelectedProfile->GetPackagingMode();
+		if (PackagingMode == ELauncherProfilePackagingModes::Locally)
+		{
+			return FText::FromString(SelectedProfile->GetPackageDirectory());
+		}
+	}
+
+	return FText::GetEmpty();
+}
+
 void SProjectLauncherPackagingSettings::HandleProfileManagerProfileSelected( const ILauncherProfilePtr& SelectedProfile, const ILauncherProfilePtr& PreviousProfile )
 {
+	UpdateDirectoryPathText();
+}
+
+FReply SProjectLauncherPackagingSettings::HandleBrowseButtonClicked()
+{
+	IDesktopPlatform* DesktopPlatform = FDesktopPlatformModule::Get();
+	if (DesktopPlatform)
+	{
+		TSharedPtr<SWindow> ParentWindow = FSlateApplication::Get().FindWidgetWindow(AsShared());
+		void* ParentWindowHandle = (ParentWindow.IsValid() && ParentWindow->GetNativeWindow().IsValid()) ? ParentWindow->GetNativeWindow()->GetOSWindowHandle() : nullptr;
+
+		FString FolderName;
+		const bool bFolderSelected = DesktopPlatform->OpenDirectoryDialog(
+			ParentWindowHandle,
+			LOCTEXT("FolderDialogTitle", "Choose a directory").ToString(),
+			DirectoryPathTextBox->GetText().ToString(),
+			FolderName
+			);
+
+		if (bFolderSelected)
+		{
+			if (!FolderName.EndsWith(TEXT("/")))
+			{
+				FolderName += TEXT("/");
+			}
+
+			DirectoryPathTextBox->SetText(FText::FromString(FolderName));
+			ILauncherProfilePtr SelectedProfile = Model->GetSelectedProfile();
+
+			if (SelectedProfile.IsValid())
+			{
+				SelectedProfile->SetPackageDirectory(FolderName);
+			}
+		}
+	}
+
+	return FReply::Handled();
+}
+
+bool SProjectLauncherPackagingSettings::IsEditable() const
+{
+	ILauncherProfilePtr SelectedProfile = Model->GetSelectedProfile();
 
+	if (SelectedProfile.IsValid())
+	{
+		return SelectedProfile->GetPackagingMode() == ELauncherProfilePackagingModes::Locally;
+	}
+	return false;
 }
 
+void SProjectLauncherPackagingSettings::OnTextChanged(const FText& InText)
+{
+	ILauncherProfilePtr SelectedProfile = Model->GetSelectedProfile();
+
+	if (SelectedProfile.IsValid())
+	{
+		SelectedProfile->SetPackageDirectory(InText.ToString());
+	}
+}
+
+void SProjectLauncherPackagingSettings::OnTextCommitted(const FText& InText, ETextCommit::Type CommitInfo)
+{
+	if (CommitInfo == ETextCommit::OnEnter)
+	{
+		ILauncherProfilePtr SelectedProfile = Model->GetSelectedProfile();
+
+		if (SelectedProfile.IsValid())
+		{
+			SelectedProfile->SetPackageDirectory(InText.ToString());
+		}
+	}
+}
 
 #undef LOCTEXT_NAMESPACE
diff --git a/Engine/Source/Developer/ProjectLauncher/Private/Widgets/Package/SProjectLauncherPackagingSettings.h b/Engine/Source/Developer/ProjectLauncher/Private/Widgets/Package/SProjectLauncherPackagingSettings.h
--- a/Engine/Source/Developer/ProjectLauncher/Private/Widgets/Package/SProjectLauncherPackagingSettings.h
+++ b/Engine/Source/Developer/ProjectLauncher/Private/Widgets/Package/SProjectLauncherPackagingSettings.h
@@ -31,16 +31,32 @@ public:
 	 */
 	void Construct( const FArguments& InArgs, const FProjectLauncherModelRef& InModel );
 
+	// Changes the directory path to the correct packaging mode version
+	void UpdateDirectoryPathText();
+
 private:
 
+	// Callback for getting the content text of the 'Directory' label.
+	FText HandleDirectoryTitleText() const;
+	FText HandleDirectoryPathText() const;
+
 	// Callback for changing the selected profile in the profile manager.
 	void HandleProfileManagerProfileSelected( const ILauncherProfilePtr& SelectedProfile, const ILauncherProfilePtr& PreviousProfile );
 
+	FReply HandleBrowseButtonClicked();
+
+	/** Handles entering in a command */
+	bool IsEditable() const;
+
+	void OnTextCommitted(const FText& InText, ETextCommit::Type CommitInfo);
+
+	void OnTextChanged(const FText& InText);
+
 private:
 
 	// Holds a pointer to the data model.
 	FProjectLauncherModelPtr Model;
 
 	// Holds the repository path text box.
-	TSharedPtr<SEditableTextBox> RepositoryPathTextBox;
+	TSharedPtr<SEditableTextBox> DirectoryPathTextBox;
 };
diff --git a/Engine/Source/Developer/SessionFrontend/Private/Models/SessionConsoleCommands.h b/Engine/Source/Developer/SessionFrontend/Private/Models/SessionConsoleCommands.h
--- a/Engine/Source/Developer/SessionFrontend/Private/Models/SessionConsoleCommands.h
+++ b/Engine/Source/Developer/SessionFrontend/Private/Models/SessionConsoleCommands.h
@@ -2,7 +2,7 @@
 
 #pragma once
 
-#define LOCTEXT_NAMESPACE ""
+#define LOCTEXT_NAMESPACE "SessionConsoleCommands"
 
 /**
  * The device details commands.
diff --git a/Engine/Source/Developer/SessionFrontend/Private/Widgets/SSessionFrontend.cpp b/Engine/Source/Developer/SessionFrontend/Private/Widgets/SSessionFrontend.cpp
--- a/Engine/Source/Developer/SessionFrontend/Private/Widgets/SSessionFrontend.cpp
+++ b/Engine/Source/Developer/SessionFrontend/Private/Widgets/SSessionFrontend.cpp
@@ -170,14 +170,14 @@ TSharedRef<SDockTab> SSessionFrontend::HandleTabManagerSpawnTab( const FSpawnTab
 		IAutomationControllerManagerPtr AutomationController = AutomationControllerModule.GetAutomationController();
 		IAutomationWindowModule& AutomationWindowModule = FModuleManager::LoadModuleChecked<IAutomationWindowModule>("AutomationWindow");
 
-		AutomationController->OnShutdown().BindRaw(this, &SSessionFrontend::HandleAutomationModuleShutdown);
+		AutomationController->OnShutdown().AddSP(this, &SSessionFrontend::HandleAutomationModuleShutdown);
 
 		TabWidget = AutomationWindowModule.CreateAutomationWindow(
 			AutomationController.ToSharedRef(),
 			SessionManager.ToSharedRef()
 		);
 
-		AutomationWindowModule.OnShutdown().BindRaw(this, &SSessionFrontend::HandleAutomationModuleShutdown);
+		AutomationWindowModule.OnShutdown().BindSP(this, &SSessionFrontend::HandleAutomationModuleShutdown);
 	}
 	else if (TabIdentifier == ProfilerTabId)
 	{
diff --git a/Engine/Source/Developer/SuperSearch/Private/SSuperSearch.cpp b/Engine/Source/Developer/SuperSearch/Private/SSuperSearch.cpp
--- a/Engine/Source/Developer/SuperSearch/Private/SSuperSearch.cpp
+++ b/Engine/Source/Developer/SuperSearch/Private/SSuperSearch.cpp
@@ -390,26 +390,23 @@ void SSuperSearchBox::OnTextChanged(const FText& InText)
 
 			for (const FAssetData& Asset : AssetData)
 			{
-				const FString* SearchTag = Asset.TagsAndValues.Find("SearchTags");
-				const FString* ResultTitle = Asset.TagsAndValues.Find("Title");
-				if( ResultTitle)
+				const FString ResultTitle = Asset.GetTagValueRef<FString>("Title");
+				if (ResultTitle.Contains(InText.ToString()))
 				{
-					if (ResultTitle->Contains(InText.ToString()))
-					{
-						FSearchEntry SearchEntry;
-						SearchEntry.Title = *ResultTitle;
-						SearchEntry.URL = "";
-						SearchEntry.bCategory = false;
-						SearchEntry.AssetData = Asset;
-						TutorialResults.Add(SearchEntry);
-					}
+					FSearchEntry SearchEntry;
+					SearchEntry.Title = ResultTitle;
+					SearchEntry.URL = "";
+					SearchEntry.bCategory = false;
+					SearchEntry.AssetData = Asset;
+					TutorialResults.Add(SearchEntry);
 				}
 
 				// If the asset has search tags, search them
-				if ((SearchTag) && (SearchTag->IsEmpty()== false))
+				const FString SearchTag = Asset.GetTagValueRef<FString>("SearchTags");
+				if (!SearchTag.IsEmpty())
 				{
 					TArray<FString> SearchTags;			
-					SearchTag->ParseIntoArray(SearchTags,TEXT(","));
+					SearchTag.ParseIntoArray(SearchTags,TEXT(","));
 					//trim any xs spaces off the strings.
 					for (int32 iTag = 0; iTag < SearchTags.Num() ; iTag++)
 					{
@@ -419,7 +416,7 @@ void SSuperSearchBox::OnTextChanged(const FText& InText)
 					if (SearchTags.Find(InText.ToString()) != INDEX_NONE)
 					{
 						FSearchEntry SearchEntry;
-						SearchEntry.Title = *ResultTitle;
+						SearchEntry.Title = ResultTitle;
 						SearchEntry.URL = "";
 						SearchEntry.bCategory = false;
 						SearchEntry.AssetData = Asset;
diff --git a/Engine/Source/Editor/AIGraph/Private/AIGraphNode.cpp b/Engine/Source/Editor/AIGraph/Private/AIGraphNode.cpp
--- a/Engine/Source/Editor/AIGraph/Private/AIGraphNode.cpp
+++ b/Engine/Source/Editor/AIGraph/Private/AIGraphNode.cpp
@@ -134,11 +134,11 @@ FText UAIGraphNode::GetTooltipText() const
 			if (NodeInstance->GetClass()->HasAnyClassFlags(CLASS_CompiledFromBlueprint))
 			{
 				FAssetData AssetData(NodeInstance->GetClass()->ClassGeneratedBy);
-				const FString* Description = AssetData.TagsAndValues.Find(GET_MEMBER_NAME_CHECKED(UBlueprint, BlueprintDescription));
-
-				if (Description && !Description->IsEmpty())
+				FString Description = AssetData.GetTagValueRef<FString>(GET_MEMBER_NAME_CHECKED(UBlueprint, BlueprintDescription));
+				if (!Description.IsEmpty())
 				{
-					TooltipDesc = FText::FromString(Description->Replace(TEXT("\\n"), TEXT("\n")));
+					Description.ReplaceInline(TEXT("\\n"), TEXT("\n"));
+					TooltipDesc = FText::FromString(MoveTemp(Description));
 				}
 			}
 			else
diff --git a/Engine/Source/Editor/AIGraph/Private/AIGraphTypes.cpp b/Engine/Source/Editor/AIGraph/Private/AIGraphTypes.cpp
--- a/Engine/Source/Editor/AIGraph/Private/AIGraphTypes.cpp
+++ b/Engine/Source/Editor/AIGraph/Private/AIGraphTypes.cpp
@@ -289,10 +289,9 @@ void FGraphNodeClassHelper::OnAssetAdded(const class FAssetData& AssetData)
 
 void FGraphNodeClassHelper::OnAssetRemoved(const class FAssetData& AssetData)
 {
-	const FString* GeneratedClassname = AssetData.TagsAndValues.Find("GeneratedClass");
-	if (GeneratedClassname)
+	FString AssetClassName;
+	if (AssetData.GetTagValue("GeneratedClass", AssetClassName))
 	{
-		FString AssetClassName = *GeneratedClassname;
 		ConstructorHelpers::StripObjectClass(AssetClassName);
 		AssetClassName = FPackageName::ObjectPathToObjectName(*AssetClassName);
 
@@ -325,16 +324,14 @@ void FGraphNodeClassHelper::OnHotReload(bool bWasTriggeredAutomatically)
 TSharedPtr<FGraphNodeClassNode> FGraphNodeClassHelper::CreateClassDataNode(const class FAssetData& AssetData)
 {
 	TSharedPtr<FGraphNodeClassNode> Node;
-	const FString* GeneratedClassname = AssetData.TagsAndValues.Find("GeneratedClass");
-	const FString* ParentClassname = AssetData.TagsAndValues.Find("ParentClass");
 
-	if (GeneratedClassname && ParentClassname)
+	FString AssetClassName;
+	FString AssetParentClassName;
+	if (AssetData.GetTagValue("GeneratedClass", AssetClassName) && AssetData.GetTagValue("ParentClass", AssetParentClassName))
 	{
-		FString AssetClassName = *GeneratedClassname;
 		UObject* Outer1(NULL);
 		ResolveName(Outer1, AssetClassName, false, false);
 
-		FString AssetParentClassName = *ParentClassname;
 		UObject* Outer2(NULL);
 		ResolveName(Outer2, AssetParentClassName, false, false);
 
diff --git a/Engine/Source/Editor/AnimGraph/Private/AnimGraphNode_SequencePlayer.cpp b/Engine/Source/Editor/AnimGraph/Private/AnimGraphNode_SequencePlayer.cpp
--- a/Engine/Source/Editor/AnimGraph/Private/AnimGraphNode_SequencePlayer.cpp
+++ b/Engine/Source/Editor/AnimGraph/Private/AnimGraphNode_SequencePlayer.cpp
@@ -142,10 +142,7 @@ void UAnimGraphNode_SequencePlayer::GetMenuActions(FBlueprintActionDatabaseRegis
 		UAnimGraphNode_SequencePlayer* SequencePlayerNode = CastChecked<UAnimGraphNode_SequencePlayer>(NewNode);
 		if (bIsTemplateNode)
 		{
-			if (const FString* SkeletonTag = AssetData.TagsAndValues.Find(TEXT("Skeleton")))
-			{
-				SequencePlayerNode->UnloadedSkeletonName = *SkeletonTag;
-			}
+			AssetData.GetTagValue("Skeleton", SequencePlayerNode->UnloadedSkeletonName);
 		}
 		else
 		{
