Fixes Missing Platform Icon in Session Frontend Automation Tab

diff --git a/Engine/Source/Developer/AutomationWindow/Private/SAutomationWindow.cpp b/Engine/Source/Developer/AutomationWindow/Private/SAutomationWindow.cpp
--- a/Engine/Source/Developer/AutomationWindow/Private/SAutomationWindow.cpp
+++ b/Engine/Source/Developer/AutomationWindow/Private/SAutomationWindow.cpp
@@ -3,6 +3,7 @@
 #include "SAutomationWindow.h"
 #include "HAL/PlatformProcess.h"
 #include "HAL/PlatformApplicationMisc.h"
+#include "PlatformInfo.h"
 #include "Misc/MessageDialog.h"
 #include "Misc/TextFilter.h"
 #include "Misc/FilterCollection.h"
@@ -1559,10 +1560,10 @@ void SAutomationWindow::RebuildPlatformIcons()
 	for (int32 ClusterIndex = 0; ClusterIndex < NumClusters; ++ClusterIndex)
 	{
 		//find the right platform icon
-		FString DeviceImageName = TEXT("Launcher.Platform_");
 		FString DeviceTypeName = AutomationController->GetDeviceTypeName(ClusterIndex);
-		DeviceImageName += DeviceTypeName;
-		const FSlateBrush* ImageToUse = FEditorStyle::GetBrush(*DeviceImageName);
+		FName DeviceImageName = PlatformInfo::FindPlatformInfo(FName(*DeviceTypeName))->DataDrivenPlatformInfo->GetIconStyleName(EPlatformIconSize::Normal);
+
+		const FSlateBrush* ImageToUse = FEditorStyle::GetBrush(DeviceImageName);
 
 		PlatformsHBox->AddSlot()
 		.AutoWidth()
@@ -1580,6 +1581,7 @@ void SAutomationWindow::RebuildPlatformIcons()
 				.ToolTipText( CreateDeviceTooltip( ClusterIndex ) )
 				[
 					SNew(SImage)
+					.DesiredSizeOverride(FVector2D(16.f, 16.f))
 					.Image(ImageToUse)
 				]
 			]
