Adjusting toolbar button settings so this legacy-style options button will show up in Session Frontend.

diff --git a/Engine/Source/Developer/AutomationWindow/Private/SAutomationWindow.cpp b/Engine/Source/Developer/AutomationWindow/Private/SAutomationWindow.cpp
--- a/Engine/Source/Developer/AutomationWindow/Private/SAutomationWindow.cpp
+++ b/Engine/Source/Developer/AutomationWindow/Private/SAutomationWindow.cpp
@@ -634,10 +634,10 @@ TSharedRef< SWidget > SAutomationWindow::MakeAutomationWindowToolBar( const TSha
 				ToolbarBuilder.AddComboButton(
 					DefaultAction,
 					FOnGetContent::CreateStatic( &SAutomationWindow::GenerateTestsOptionsMenuContent, InAutomationWindow ),
-					LOCTEXT( "TestOptions_Label", "Test Options" ),
+					TAttribute<FText>(),
 					LOCTEXT( "TestOptionsToolTip", "Test Options" ),
 					FSlateIcon(FEditorStyle::GetStyleSetName(), "AutomationWindow.TestOptions"),
-					true);
+					false);
 
 				// Added button for running the currently open level test.
 #if WITH_EDITOR
