Mark test category as complete after all children finished running.

diff --git a/Engine/Source/Developer/AutomationWindow/Private/SAutomationTestItem.cpp b/Engine/Source/Developer/AutomationWindow/Private/SAutomationTestItem.cpp
--- a/Engine/Source/Developer/AutomationWindow/Private/SAutomationTestItem.cpp
+++ b/Engine/Source/Developer/AutomationWindow/Private/SAutomationTestItem.cpp
@@ -232,15 +232,36 @@ TSharedRef<SWidget> SAutomationTestItem::GenerateWidgetForColumn( const FName& C
 				[
 					SNew(SBorder)
 					.BorderImage( FEditorStyle::GetBrush("ErrorReporting.Box") )
-					.HAlign(HAlign_Fill)
+					.HAlign(HAlign_Center)
 					.VAlign(VAlign_Center)
 					.Padding( FMargin(3,0) )
 					.BorderBackgroundColor( FSlateColor( FLinearColor( 1.0f, 0.0f, 1.0f, 0.0f ) ) )
 					[
-						//progress bar for percent of enabled children completed
-						SNew(SProgressBar)
-						.Percent( this, &SAutomationTestItem::ItemStatus_ProgressFraction, ClusterIndex )
-						.FillColorAndOpacity(this, &SAutomationTestItem::ItemStatus_ProgressColor, ClusterIndex )
+						SNew(SHorizontalBox)
+						+SHorizontalBox::Slot()
+						.AutoWidth()
+						.HAlign(HAlign_Center)
+						.VAlign(VAlign_Center)
+						[
+							//image when children complete or not run
+							SNew(SImage)
+							.Image(this, &SAutomationTestItem::ItemChildrenStatus_StatusImage, ClusterIndex)
+							.Visibility(this, &SAutomationTestItem::ItemStatus_GetChildrenStatusVisibility, ClusterIndex, false)
+						]
+						+SHorizontalBox::Slot()
+						.AutoWidth()
+						[
+							SNew(SBox)
+							.WidthOverride(ColumnWidth - 8)
+							.HeightOverride(16.0f)
+							[
+								//progress bar for percent of enabled children completed
+								SNew(SProgressBar)
+								.Percent(this, &SAutomationTestItem::ItemStatus_ProgressFraction, ClusterIndex)
+								.FillColorAndOpacity(this, &SAutomationTestItem::ItemStatus_ProgressColor, ClusterIndex)
+								.Visibility(this, &SAutomationTestItem::ItemStatus_GetChildrenStatusVisibility, ClusterIndex, true)
+							]
+						]
 					]
 				];
 			}
@@ -483,6 +504,26 @@ EVisibility SAutomationTestItem::ItemStatus_GetStatusVisibility(const int32 Clus
 	return bFinalVisibility ? EVisibility::Visible : EVisibility::Collapsed;
 }
 
+EVisibility SAutomationTestItem::ItemStatus_GetChildrenStatusVisibility(const int32 ClusterIndex, const bool bForInProcessThrobber) const
+{
+	// Internal node: completion status image visible when all children completed
+	bool bImageVisible = false;
+
+	FAutomationCompleteState CompleteState;
+	const int32 PassIndex = TestStatus->GetCurrentPassIndex(ClusterIndex);
+	TestStatus->GetCompletionStatus(ClusterIndex, PassIndex, CompleteState);
+
+	uint32 TotalComplete = CompleteState.NumEnabledTestsPassed + CompleteState.NumEnabledTestsFailed + CompleteState.NumEnabledTestsCouldntBeRun;
+	if ((TotalComplete > 0) && (CompleteState.TotalEnabled > 0))
+	{
+		bImageVisible = (TotalComplete == CompleteState.TotalEnabled);
+	}
+
+	bool bFinalVisibility = bForInProcessThrobber ? !bImageVisible : bImageVisible;
+
+	return bFinalVisibility ? EVisibility::Visible : EVisibility::Collapsed;
+}
+
 
 FText SAutomationTestItem::ItemStatus_NumParticipantsRequiredText() const
 {
@@ -589,6 +630,32 @@ const FSlateBrush* SAutomationTestItem::ItemStatus_StatusImage(const int32 Clust
 	return ImageToUse;
 }
 
+const FSlateBrush* SAutomationTestItem::ItemChildrenStatus_StatusImage(const int32 ClusterIndex) const
+{
+	FAutomationCompleteState CompleteState;
+	const int32 PassIndex = TestStatus->GetCurrentPassIndex(ClusterIndex);
+	TestStatus->GetCompletionStatus(ClusterIndex, PassIndex, CompleteState);
+
+	const FSlateBrush* ImageToUse = FEditorStyle::GetBrush("Automation.InProcess");
+
+	uint32 TotalComplete = CompleteState.NumEnabledTestsPassed + CompleteState.NumEnabledTestsFailed + CompleteState.NumEnabledTestsCouldntBeRun;
+	if ((TotalComplete > 0) && (CompleteState.TotalEnabled > 0) && TotalComplete == CompleteState.TotalEnabled) {
+		if (TotalComplete == CompleteState.NumEnabledTestsPassed)
+		{
+			ImageToUse = FEditorStyle::GetBrush("Automation.Success");
+		}
+		else if (CompleteState.NumEnabledTestsFailed)
+		{
+			ImageToUse = FEditorStyle::GetBrush("Automation.Fail");
+		}
+		else if (CompleteState.NumEnabledTestsCouldntBeRun)
+		{
+			ImageToUse = FEditorStyle::GetBrush("Automation.NotRun");
+		}
+	}
+	return ImageToUse;
+}
+
 /* SAutomationTestitem event handlers
  *****************************************************************************/
 
diff --git a/Engine/Source/Developer/AutomationWindow/Private/SAutomationTestItem.h b/Engine/Source/Developer/AutomationWindow/Private/SAutomationTestItem.h
--- a/Engine/Source/Developer/AutomationWindow/Private/SAutomationTestItem.h
+++ b/Engine/Source/Developer/AutomationWindow/Private/SAutomationTestItem.h
@@ -144,6 +144,15 @@ protected:
 	 */
 	EVisibility ItemStatus_GetStatusVisibility( const int32 ClusterIndex, const bool bForInProcessThrobber ) const;
 
+	/**
+	 * Helper for internal nodes to ensure throbber is visible when children are in process and icon is visible when children completed.
+	 *
+	 * @param ClusterIndex The cluster index.
+	 * @param bForInProcessThrobber If throbbing is in process.
+	 * @return The visibility state.
+	 */
+	EVisibility ItemStatus_GetChildrenStatusVisibility( const int32 ClusterIndex, const bool bForInProcessThrobber ) const;
+
 	/**
 	 * Color of progress bar for internal tree test nodes.
 	 *
@@ -167,6 +176,14 @@ protected:
 	 * @return the status image.
 	 */
 	const FSlateBrush* ItemStatus_StatusImage( const int32 ClusterIndex ) const;
+
+	/**
+	 * Returns image that denotes the status of a test category as implied by its children, on the given platform cluster.
+	 *
+	 * @param ClusterIndex The cluster index.
+	 * @return the status image.
+	 */
+	const FSlateBrush* ItemChildrenStatus_StatusImage(const int32 ClusterIndex) const;
 	
 	/**
 	 * The number of participants required for this test item in string form.
