Correction to LOCTEXT macro that do not support using an empty id.

diff --git a/Engine/Source/Developer/AutomationWindow/Private/SAutomationTestItem.cpp b/Engine/Source/Developer/AutomationWindow/Private/SAutomationTestItem.cpp
--- a/Engine/Source/Developer/AutomationWindow/Private/SAutomationTestItem.cpp
+++ b/Engine/Source/Developer/AutomationWindow/Private/SAutomationTestItem.cpp
@@ -392,7 +392,7 @@ FText SAutomationTestItem::GetExcludeReason() const
 	FName Reason;
 	bool IsToBeSkipped = TestStatus->IsToBeSkipped(&Reason);
 
-	return IsToBeSkipped ? FText::FromName(Reason) : LOCTEXT("", "");
+	return IsToBeSkipped ? FText::FromName(Reason) : FText();
 }
 
 void SAutomationTestItem::SetSkipFlag(ECheckBoxState Enable)
