Prevent UnrealFrontend crash if platform DDPI lookup fails, using an empty icon instead

diff --git a/Engine/Source/Developer/AutomationWindow/Private/SAutomationWindow.cpp b/Engine/Source/Developer/AutomationWindow/Private/SAutomationWindow.cpp
--- a/Engine/Source/Developer/AutomationWindow/Private/SAutomationWindow.cpp
+++ b/Engine/Source/Developer/AutomationWindow/Private/SAutomationWindow.cpp
@@ -1574,9 +1574,13 @@ void SAutomationWindow::RebuildPlatformIcons()
 	{
 		//find the right platform icon
 		FString DeviceTypeName = AutomationController->GetDeviceTypeName(ClusterIndex);
-		FName DeviceImageName = PlatformInfo::FindPlatformInfo(FName(*DeviceTypeName))->DataDrivenPlatformInfo->GetIconStyleName(EPlatformIconSize::Normal);
+		const PlatformInfo::FTargetPlatformInfo* TargetPlatformInfo = PlatformInfo::FindPlatformInfo(FName(*DeviceTypeName));
+		ensureMsgf(TargetPlatformInfo, TEXT("PlatformInfo::FindPlatformInfo() failed looking up %s"), *DeviceTypeName);
+		ensureMsgf(!TargetPlatformInfo || TargetPlatformInfo->DataDrivenPlatformInfo, TEXT("PlatformInfo::FTargetPlatformInfo had null DataDrivenPlatformInfo"));
+		const FDataDrivenPlatformInfo* DDPI = TargetPlatformInfo ? TargetPlatformInfo->DataDrivenPlatformInfo : nullptr;
+		FName DeviceImageName = DDPI ? DDPI->GetIconStyleName(EPlatformIconSize::Normal) : NAME_None;
 
-		const FSlateBrush* ImageToUse = FEditorStyle::GetBrush(DeviceImageName);
+		const FSlateBrush* ImageToUse = FAppStyle::Get().GetBrush(DeviceImageName);
 
 		PlatformsHBox->AddSlot()
 		.AutoWidth()
