Crash fix when no groups configured
#rnx

diff --git a/Engine/Source/Developer/AutomationWindow/Private/SAutomationWindow.cpp b/Engine/Source/Developer/AutomationWindow/Private/SAutomationWindow.cpp
--- a/Engine/Source/Developer/AutomationWindow/Private/SAutomationWindow.cpp
+++ b/Engine/Source/Developer/AutomationWindow/Private/SAutomationWindow.cpp
@@ -283,21 +283,14 @@ void SAutomationWindow::Construct( const FArguments& InArgs, const IAutomationCo
 	TArray<FAutomatedTestGroup> TestGroups = Settings->Groups;
 	GroupComboList.Empty();
 	GroupFiltersMap.Empty();
-	if (TestGroups.Num() > 0)
-	{
-		const FString AllGroups = TEXT("All Groups");
-		GroupComboList.Add(MakeShareable(new FString(AllGroups)));
-		GroupFiltersMap.Add(AllGroups, TArray<FAutomatedTestFilter>());
-		for (int TestGroupIdx = 0; TestGroupIdx < TestGroups.Num(); TestGroupIdx++)
-		{
-			GroupComboList.Add(MakeShareable(new FString(TestGroups[TestGroupIdx].Name)));
-			GroupFiltersMap.Add(TestGroups[TestGroupIdx].Name, TestGroups[TestGroupIdx].Filters);
-		}
-	}
-	else
-	{
-		GroupComboBox->SetVisibility(EVisibility::Collapsed);
-	}
+	const FString AllGroups = TEXT("All Groups");
+	GroupComboList.Add(MakeShareable(new FString(AllGroups)));
+	GroupFiltersMap.Add(AllGroups, TArray<FAutomatedTestFilter>());
+	for (int TestGroupIdx = 0; TestGroupIdx < TestGroups.Num(); TestGroupIdx++)
+	{
+		GroupComboList.Add(MakeShareable(new FString(TestGroups[TestGroupIdx].Name)));
+		GroupFiltersMap.Add(TestGroups[TestGroupIdx].Name, TestGroups[TestGroupIdx].Filters);
+	}	
 	
 	TSharedRef<SNotificationList> NotificationList = SNew(SNotificationList) .Visibility( EVisibility::HitTestInvisible );
 
@@ -383,6 +376,7 @@ void SAutomationWindow::Construct( const FArguments& InArgs, const IAutomationCo
 											.MinDesiredWidth(130.0f)
 											[
 												SAssignNew(GroupComboBox, SComboBox< TSharedPtr<FString> >)
+												.Visibility(this, &SAutomationWindow::HandleGroupsVisibility)
 												.OptionsSource(&GroupComboList)
 												.InitiallySelectedItem(GroupComboList[0])
 												.OnGenerateWidget(this, &SAutomationWindow::GenerateGroupComboItem)
@@ -923,6 +917,11 @@ EVisibility SAutomationWindow::HandlePresetTextVisibility( ) const
 	return bAddingTestPreset ? EVisibility::Visible : EVisibility::Hidden;
 }
 
+EVisibility SAutomationWindow::HandleGroupsVisibility() const
+{
+	return GroupComboList.Num() == 1 ? EVisibility::Collapsed : EVisibility::Visible;
+}
+
 bool SAutomationWindow::IsAddButtonEnabled() const
 {
 	return !bAddingTestPreset && IsAutomationControllerIdle();
diff --git a/Engine/Source/Developer/AutomationWindow/Private/SAutomationWindow.h b/Engine/Source/Developer/AutomationWindow/Private/SAutomationWindow.h
--- a/Engine/Source/Developer/AutomationWindow/Private/SAutomationWindow.h
+++ b/Engine/Source/Developer/AutomationWindow/Private/SAutomationWindow.h
@@ -472,6 +472,9 @@ private:
 	/** Handles if the add preset text box should be visible. */
 	EVisibility HandlePresetTextVisibility( ) const;
 
+	/** Handles if the groups combo box should be visible. */
+	EVisibility HandleGroupsVisibility() const;
+
 	/** Called when the user commits the text in the add preset text box. */
 	void HandlePresetTextCommited( const FText& CommittedText, ETextCommit::Type CommitType );
 
