Fix status icons in Test Automation window on sub category

#rnx

diff --git a/Engine/Source/Developer/AutomationWindow/Private/SAutomationTestItem.cpp b/Engine/Source/Developer/AutomationWindow/Private/SAutomationTestItem.cpp
--- a/Engine/Source/Developer/AutomationWindow/Private/SAutomationTestItem.cpp
+++ b/Engine/Source/Developer/AutomationWindow/Private/SAutomationTestItem.cpp
@@ -630,21 +630,21 @@
 	const int32 PassIndex = TestStatus->GetCurrentPassIndex(ClusterIndex);
 	TestStatus->GetCompletionStatus(ClusterIndex, PassIndex, CompleteState);
 
-	const FSlateBrush* ImageToUse = FAppStyle::GetBrush("Automation.InProcess");
+	const FSlateBrush* ImageToUse = FAutomationWindowStyle::Get().GetBrush("Automation.InProcess");
 
 	uint32 TotalComplete = CompleteState.NumEnabledTestsPassed + CompleteState.NumEnabledTestsFailed + CompleteState.NumEnabledTestsCouldntBeRun;
 	if ((TotalComplete > 0) && (CompleteState.TotalEnabled > 0) && TotalComplete == CompleteState.TotalEnabled) {
 		if (TotalComplete == CompleteState.NumEnabledTestsPassed)
 		{
-			ImageToUse = FAppStyle::GetBrush("Automation.Success");
+			ImageToUse = FAutomationWindowStyle::Get().GetBrush("Automation.Success");
 		}
 		else if (CompleteState.NumEnabledTestsFailed)
 		{
-			ImageToUse = FAppStyle::GetBrush("Automation.Fail");
+			ImageToUse = FAutomationWindowStyle::Get().GetBrush("Automation.Fail");
 		}
 		else if (CompleteState.NumEnabledTestsCouldntBeRun)
 		{
-			ImageToUse = FAppStyle::GetBrush("Automation.NotRun");
+			ImageToUse = FAutomationWindowStyle::Get().GetBrush("Automation.NotRun");
 		}
 	}
 	return ImageToUse;
