AutomationWindow: When AutomationWindow is active during editor start up, fix the infiinite loop caused by ListTests being triggered from every ScanPathsSynchronous and then calling another ScanPathsSynchronous.
#rnx

diff --git a/Engine/Source/Developer/AutomationWindow/Private/SAutomationWindow.cpp b/Engine/Source/Developer/AutomationWindow/Private/SAutomationWindow.cpp
--- a/Engine/Source/Developer/AutomationWindow/Private/SAutomationWindow.cpp
+++ b/Engine/Source/Developer/AutomationWindow/Private/SAutomationWindow.cpp
@@ -148,7 +148,7 @@ void SAutomationWindow::Construct( const FArguments& InArgs, const IAutomationCo
 
 #if WITH_EDITOR
 	FAssetRegistryModule& AssetRegistryModule = FModuleManager::LoadModuleChecked<FAssetRegistryModule>(TEXT("AssetRegistry"));
-	AssetRegistryModule.Get().OnFileLoadProgressUpdated().AddSP(this, &SAutomationWindow::OnAssetRegistryFileLoadProgress);
+	AssetRegistryModule.Get().OnFilesLoaded().AddSP(this, &SAutomationWindow::OnAssetRegistryFilesLoaded);
 #endif
 
 	TestPresetManager = MakeShareable(new FAutomationTestPresetManager());
@@ -2355,12 +2355,9 @@ bool SAutomationWindow::HandleMainContentIsEnabled() const
 #if WITH_EDITOR
 // React to asset registry finishing updating.
 // We only want to do this if there are no tests already listed, otherwise this fires every time you save a map for example.
-void SAutomationWindow::OnAssetRegistryFileLoadProgress(const IAssetRegistry::FFileLoadProgressUpdateData& ProgressUpdateData)
+void SAutomationWindow::OnAssetRegistryFilesLoaded()
 {
-	if ( ProgressUpdateData.NumAssetsProcessedByAssetRegistry == ProgressUpdateData.NumTotalAssets && IsAutomationControllerIdle() && AutomationController->GetReports().Num() == 0 )
-	{
-		ListTests();
-	}
+	ListTests();
 }
 #endif
 
diff --git a/Engine/Source/Developer/AutomationWindow/Private/SAutomationWindow.h b/Engine/Source/Developer/AutomationWindow/Private/SAutomationWindow.h
--- a/Engine/Source/Developer/AutomationWindow/Private/SAutomationWindow.h
+++ b/Engine/Source/Developer/AutomationWindow/Private/SAutomationWindow.h
@@ -543,7 +543,7 @@ private:
 	}
 
 #if WITH_EDITOR
-	void OnAssetRegistryFileLoadProgress(const IAssetRegistry::FFileLoadProgressUpdateData& ProgressUpdateData);
+	void OnAssetRegistryFilesLoaded();
 #endif
 
 private:
