Test Automation - Add Keep PIE open in advance settings
+ add project config settings to change that default.

diff --git a/Engine/Source/Developer/AutomationController/Private/AutomationControllerManager.cpp b/Engine/Source/Developer/AutomationController/Private/AutomationControllerManager.cpp
--- a/Engine/Source/Developer/AutomationController/Private/AutomationControllerManager.cpp
+++ b/Engine/Source/Developer/AutomationController/Private/AutomationControllerManager.cpp
@@ -163,6 +163,8 @@ FAutomationControllerManager::FAutomationControllerManager()
 	{
 		GameInstanceLostTimerSeconds = Settings->GameInstanceLostTimerSeconds;
 	}
+
+	bKeepPIEOpen = Settings->bKeepPIEOpen;
 	
 	FString DeveloperPath;
 	FParse::Value(FCommandLine::Get(), TEXT("ReportOutputPath="), ReportExportPath, false);
@@ -406,7 +408,7 @@ void FAutomationControllerManager::StopTests()
 
 	// Close play window
 #if WITH_EDITOR
-	if (GUnrealEd)
+	if (GUnrealEd && !bKeepPIEOpen)
 	{
 		GUnrealEd->RequestEndPlayMap();
 	}
@@ -533,7 +535,7 @@ void FAutomationControllerManager::ProcessAvailableTasks()
 
 					// Close play window
 					#if WITH_EDITOR
-					if (GUnrealEd)
+					if (GUnrealEd && !bKeepPIEOpen)
 					{
 						GUnrealEd->RequestEndPlayMap();
 					}
diff --git a/Engine/Source/Developer/AutomationController/Private/AutomationControllerManager.h b/Engine/Source/Developer/AutomationController/Private/AutomationControllerManager.h
--- a/Engine/Source/Developer/AutomationController/Private/AutomationControllerManager.h
+++ b/Engine/Source/Developer/AutomationController/Private/AutomationControllerManager.h
@@ -229,6 +229,16 @@ public:
 		bSendAnalytics = bNewValue;
 	}
 
+	virtual bool KeepPIEOpen() const override
+	{
+		return bKeepPIEOpen;
+	}
+
+	virtual void SetKeepPIEOpen(const bool bNewValue) override
+	{
+		bKeepPIEOpen = bNewValue;
+	}
+
 	virtual void SetFilter( TSharedPtr< AutomationFilterCollection > InFilter ) override
 	{
 		ReportManager.SetFilter( InFilter );
@@ -578,6 +588,9 @@ private:
 	/** If we should send result to analytics */
 	bool bSendAnalytics = false;
 
+	/** If we should send keep the PIE open when test pass end */
+	bool bKeepPIEOpen = false;
+
 	/** The list of results generated by our test pass. */
 	FAutomatedTestPassResults JsonTestPassResults;
 
diff --git a/Engine/Source/Developer/AutomationController/Private/AutomationControllerSettings.cpp b/Engine/Source/Developer/AutomationController/Private/AutomationControllerSettings.cpp
--- a/Engine/Source/Developer/AutomationController/Private/AutomationControllerSettings.cpp
+++ b/Engine/Source/Developer/AutomationController/Private/AutomationControllerSettings.cpp
@@ -10,6 +10,7 @@ UAutomationControllerSettings::UAutomationControllerSettings(const FObjectInitia
 	, bSuppressLogErrors(false)
 	, bSuppressLogWarnings(false)
 	, bElevateLogWarningsToErrors(true)
+	, bKeepPIEOpen(false)
 	, bTreatLogWarningsAsTestErrors(true)
 	, GameInstanceLostTimerSeconds(300.0f)
 	, bResetTelemetryStorageOnNewSession(false)
diff --git a/Engine/Source/Developer/AutomationController/Public/AutomationControllerSettings.h b/Engine/Source/Developer/AutomationController/Public/AutomationControllerSettings.h
--- a/Engine/Source/Developer/AutomationController/Public/AutomationControllerSettings.h
+++ b/Engine/Source/Developer/AutomationController/Public/AutomationControllerSettings.h
@@ -92,7 +92,11 @@ public:
 	/** Log categories where warnings/errors will not affect the result of tests. A finer-grained way of preventing rogue systems from leading to test warnings/errors */
 	UPROPERTY(Config)
 	TArray<FString> SuppressedLogCategories;
-	
+
+	/** Whether to keep the PIE Open in the editor at the end of a test pass (default=false) */
+	UPROPERTY(Config)
+	bool bKeepPIEOpen;
+
 private:
 	/** Whether to treat log warnings as test errors (default=true) */
 	UPROPERTY(Config, Meta = (DeprecatedProperty, DeprecationMessage = "Use bElevateLogWarningsToErrors instead."))
diff --git a/Engine/Source/Developer/AutomationController/Public/IAutomationControllerManager.h b/Engine/Source/Developer/AutomationController/Public/IAutomationControllerManager.h
--- a/Engine/Source/Developer/AutomationController/Public/IAutomationControllerManager.h
+++ b/Engine/Source/Developer/AutomationController/Public/IAutomationControllerManager.h
@@ -223,6 +223,16 @@ public:
 	*/
 	virtual void SetSendAnalytics(const bool bNewValue) = 0;
 
+	/**
+	* Returns if PIE should be kept open when test pass end
+	*/
+	virtual bool KeepPIEOpen() const = 0;
+
+	/**
+	* Sets if PIE should be kept open when test pass end
+	*/
+	virtual void SetKeepPIEOpen(const bool bNewValue) = 0;
+
 	/**
 	 * Filters the visible list of tests.
 	 */
diff --git a/Engine/Source/Developer/AutomationWindow/Private/SAutomationWindow.cpp b/Engine/Source/Developer/AutomationWindow/Private/SAutomationWindow.cpp
--- a/Engine/Source/Developer/AutomationWindow/Private/SAutomationWindow.cpp
+++ b/Engine/Source/Developer/AutomationWindow/Private/SAutomationWindow.cpp
@@ -1271,10 +1271,25 @@ TSharedRef< SWidget > SAutomationWindow::GenerateTestsOptionsMenuContent( )
 		];
 
 
+	TSharedRef<SWidget> KeepPIEOpenWidget =
+		SNew(SCheckBox)
+		.IsChecked(this, &SAutomationWindow::KeepPIEOpenCheckBoxChecked)
+		.OnCheckStateChanged(this, &SAutomationWindow::HandleKeepPIEOpenBoxCheckStateChanged)
+		.Padding(FMargin(4.0f, 0.0f))
+		.ToolTipText(LOCTEXT("AutomationKeepPIEOpenTip", "If checked, the PIE will be kept open at the end of the test pass"))
+		.IsEnabled(this, &SAutomationWindow::IsAutomationControllerIdle)
+		.Content()
+		[
+			SNew(STextBlock)
+			.Text(LOCTEXT("AutomationKeepPIEOpenText", "Keep PIE open at the end"))
+		];
+
+
 	MenuBuilder.BeginSection("AutomationWindowRunTest", LOCTEXT("RunTestOptions", "Advanced Settings"));
 	{
 		MenuBuilder.AddWidget(NumTests, FText::GetEmpty());
 		MenuBuilder.AddWidget(SendAnalyticsWidget, FText::GetEmpty());
+		MenuBuilder.AddWidget(KeepPIEOpenWidget, FText::GetEmpty());
 	}
 	MenuBuilder.EndSection();
 
@@ -1291,6 +1306,16 @@ void SAutomationWindow::HandleSendAnalyticsBoxCheckStateChanged(ECheckBoxState C
 	AutomationController->SetSendAnalytics(CheckBoxState == ECheckBoxState::Checked);
 }
 
+ECheckBoxState SAutomationWindow::KeepPIEOpenCheckBoxChecked() const
+{
+	return AutomationController->KeepPIEOpen() ? ECheckBoxState::Checked : ECheckBoxState::Unchecked;
+}
+
+void SAutomationWindow::HandleKeepPIEOpenBoxCheckStateChanged(ECheckBoxState CheckBoxState)
+{
+	AutomationController->SetKeepPIEOpen(CheckBoxState == ECheckBoxState::Checked);
+}
+
 TArray<FString> SAutomationWindow::SaveExpandedTestNames(TSet<TSharedPtr<IAutomationReport>> ExpandedItems)
 {
 	TArray<FString> ExpandedItemsNames;
diff --git a/Engine/Source/Developer/AutomationWindow/Private/SAutomationWindow.h b/Engine/Source/Developer/AutomationWindow/Private/SAutomationWindow.h
--- a/Engine/Source/Developer/AutomationWindow/Private/SAutomationWindow.h
+++ b/Engine/Source/Developer/AutomationWindow/Private/SAutomationWindow.h
@@ -321,12 +321,18 @@ private:
 	/** Toggles filtering of tests based on error condition */
 	void OnToggleErrorFilter();
 	
-	/** Returns if analytics should be sent to the back end*/
+	/** Returns if analytics should be sent to the back end */
 	ECheckBoxState IsSendAnalyticsCheckBoxChecked() const;
 
-	/** Toggles if we are sending analytics results from the tests*/
+	/** Toggles if we are sending analytics results from the tests */
 	void HandleSendAnalyticsBoxCheckStateChanged(ECheckBoxState CheckBoxState);
 
+	/** Returns if PIE should be kept open when test pass end */
+	ECheckBoxState KeepPIEOpenCheckBoxChecked() const;
+
+	/** Toggles if PIE should be kept open when test pass end */
+	void HandleKeepPIEOpenBoxCheckStateChanged(ECheckBoxState CheckBoxState);
+
 	/** Returns if a device group is enabled */
 	ECheckBoxState IsDeviceGroupCheckBoxIsChecked(const int32 DeviceGroupFlag) const;
 	
