Fix bad alignment with the AutomationWindow. It was not using the base constructor.

diff --git a/Engine/Source/Developer/AutomationWindow/Private/SAutomationTestTreeView.h b/Engine/Source/Developer/AutomationWindow/Private/SAutomationTestTreeView.h
deleted file mode 100644
diff --git a/Engine/Source/Developer/AutomationWindow/Private/SAutomationWindow.cpp b/Engine/Source/Developer/AutomationWindow/Private/SAutomationWindow.cpp
--- a/Engine/Source/Developer/AutomationWindow/Private/SAutomationWindow.cpp
+++ b/Engine/Source/Developer/AutomationWindow/Private/SAutomationWindow.cpp
@@ -182,9 +182,9 @@ void SAutomationWindow::Construct( const FArguments& InArgs, const IAutomationCo
 	bIsRequestingTests = false;
 
 	//make the widget for platforms
-	PlatformsHBox = SNew (SHorizontalBox);
+	PlatformsHBox = SNew(SHorizontalBox);
 
-	TestTable = SNew(SAutomationTestTreeView< TSharedPtr< IAutomationReport > >)
+	TestTable = SNew(STreeView< TSharedPtr< IAutomationReport > >)
 		.SelectionMode(ESelectionMode::Multi)
 		.TreeItemsSource( &TestReports )
 		// Generates the actual widget for a tree item
@@ -203,69 +203,69 @@ void SAutomationWindow::Construct( const FArguments& InArgs, const IAutomationCo
 #endif
 		.HeaderRow
 		(
-		SAssignNew(TestTableHeaderRow,SHeaderRow)
-		+ SHeaderRow::Column( AutomationTestWindowConstants::Title )
-		.FillWidth(0.80f)
-		[
-			SNew(SHorizontalBox)
-			+SHorizontalBox::Slot()
-			.AutoWidth()
-			.HAlign(HAlign_Center)
+			SNew(SHeaderRow)
+			+ SHeaderRow::Column( AutomationTestWindowConstants::Title )
+			.FillWidth(0.80f)
 			[
-				//global enable/disable check box
-				SAssignNew(HeaderCheckbox, SCheckBox)
-				.OnCheckStateChanged( this, &SAutomationWindow::HeaderCheckboxStateChange)
-				.ToolTipText( LOCTEXT( "Enable Disable Test", "Enable / Disable  Test" ) )
+				SNew(SHorizontalBox)
+				+SHorizontalBox::Slot()
+				.AutoWidth()
+				.HAlign(HAlign_Center)
+				[
+					//global enable/disable check box
+					SNew(SCheckBox)
+					.OnCheckStateChanged( this, &SAutomationWindow::HeaderCheckboxStateChange)
+					.ToolTipText( LOCTEXT( "Enable Disable Test", "Enable / Disable  Test" ) )
+				]
+				+SHorizontalBox::Slot()
+				.AutoWidth()
+				.VAlign(VAlign_Center)
+				[
+					SNew( STextBlock )
+					.Text( LOCTEXT("TestName_Header", "Test Name") )
+				]
 			]
-			+SHorizontalBox::Slot()
-			.AutoWidth()
-			.VAlign(VAlign_Center)
+	
+			+ SHeaderRow::Column( AutomationTestWindowConstants::SmokeTest )
+			.FixedWidth( 50.0f )
+			.HAlignHeader(HAlign_Center)
+			.VAlignHeader(VAlign_Center)
+			.HAlignCell(HAlign_Center)
+			.VAlignCell(VAlign_Center)
 			[
-				SNew( STextBlock )
-				.Text( LOCTEXT("TestName_Header", "Test Name") )
+				//icon for the smoke test column
+				SNew(SImage)
+				.ColorAndOpacity(FLinearColor(1.0f, 1.0f, 1.0f, 0.4f))
+				.ToolTipText( LOCTEXT( "Smoke Test", "Smoke Test" ) )
+				.Image(FAutomationWindowStyle::Get().GetBrush("Automation.SmokeTest"))
 			]
-		]
-
-		+ SHeaderRow::Column( AutomationTestWindowConstants::SmokeTest )
-		.FixedWidth( 50.0f )
-		.HAlignHeader(HAlign_Center)
-		.VAlignHeader(VAlign_Center)
-		.HAlignCell(HAlign_Center)
-		.VAlignCell(VAlign_Center)
-		[
-			//icon for the smoke test column
-			SNew(SImage)
-			.ColorAndOpacity(FLinearColor(1.0f, 1.0f, 1.0f, 0.4f))
-			.ToolTipText( LOCTEXT( "Smoke Test", "Smoke Test" ) )
-			.Image(FAutomationWindowStyle::Get().GetBrush("Automation.SmokeTest"))
-		]
-
-		+ SHeaderRow::Column( AutomationTestWindowConstants::RequiredDeviceCount )
-		.FixedWidth(50.0f)
-		.HAlignHeader(HAlign_Center)
-		.VAlignHeader(VAlign_Center)
-		.HAlignCell(HAlign_Center)
-		.VAlignCell(VAlign_Center)
-		[
-			SNew( SImage )
-			.Image(FAutomationWindowStyle::Get().GetBrush("Automation.ParticipantsWarning") )
-			.ToolTipText( LOCTEXT( "RequiredDeviceCountWarningToolTip", "Number of devices required." ) )
-		]
-
-		+ SHeaderRow::Column(AutomationTestWindowConstants::Timing)
-		.FixedWidth(100.0f)
-		.DefaultLabel(LOCTEXT("TestDurationRange", "Duration"))
+	
+			+ SHeaderRow::Column( AutomationTestWindowConstants::RequiredDeviceCount )
+			.FixedWidth(50.0f)
+			.HAlignHeader(HAlign_Center)
+			.VAlignHeader(VAlign_Center)
+			.HAlignCell(HAlign_Center)
+			.VAlignCell(VAlign_Center)
+			[
+				SNew( SImage )
+				.Image(FAutomationWindowStyle::Get().GetBrush("Automation.ParticipantsWarning") )
+				.ToolTipText( LOCTEXT( "RequiredDeviceCountWarningToolTip", "Number of devices required." ) )
+			]
+	
+			+ SHeaderRow::Column(AutomationTestWindowConstants::Timing)
+			.FixedWidth(100.0f)
+			.DefaultLabel(LOCTEXT("TestDurationRange", "Duration"))
 
-		+ SHeaderRow::Column( AutomationTestWindowConstants::Status )
-		.FixedWidth(50.0f)
-		[
-			//platform header placeholder
-			PlatformsHBox.ToSharedRef()
-		]
+			+ SHeaderRow::Column( AutomationTestWindowConstants::Status )
+			.FixedWidth(50.0f)
+			[
+				//platform header placeholder
+				PlatformsHBox.ToSharedRef()
+			]
 
-		+ SHeaderRow::Column(AutomationTestWindowConstants::IsToBeSkipped)
-		.FillWidth(0.10f)
-		.DefaultLabel(LOCTEXT("Excluded", "Excluded"))
+			+ SHeaderRow::Column(AutomationTestWindowConstants::IsToBeSkipped)
+			.FillWidth(0.10f)
+			.DefaultLabel(LOCTEXT("Excluded", "Excluded"))
 
 		);
 
@@ -1207,7 +1207,7 @@ void SAutomationWindow::HandleDeviceGroupCheckStateChanged(ECheckBoxState CheckB
 	RebuildPlatformIcons();
 
 	//Need to force the tree to do a full refresh here because the reports have changed but the tree will keep using cached data.
-	TestTable->ReCreateTreeView();
+	TestTable->RebuildList();
 }
 
 TSharedRef< SWidget > SAutomationWindow::GenerateTestsOptionsMenuContent( TWeakPtr<class SAutomationWindow> InAutomationWindow )
@@ -1327,7 +1327,7 @@ TArray<FString> SAutomationWindow::SaveExpandedTestNames(TSet<TSharedPtr<IAutoma
 }
 
 // Expanded the given item if its name is in the array of strings given.
-void SAutomationWindow::ExpandItemsInList(TSharedPtr<SAutomationTestTreeView<TSharedPtr<IAutomationReport>>> InTestTable, TSharedPtr<IAutomationReport> InReport, TArray<FString> ItemsToExpand)
+void SAutomationWindow::ExpandItemsInList(TSharedPtr<STreeView<TSharedPtr<IAutomationReport>>> InTestTable, TSharedPtr<IAutomationReport> InReport, TArray<FString> ItemsToExpand)
 {
 	InTestTable->SetItemExpansion(InReport, ItemsToExpand.Contains(InReport->GetDisplayNameWithDecoration()));
 
diff --git a/Engine/Source/Developer/AutomationWindow/Private/SAutomationWindow.h b/Engine/Source/Developer/AutomationWindow/Private/SAutomationWindow.h
--- a/Engine/Source/Developer/AutomationWindow/Private/SAutomationWindow.h
+++ b/Engine/Source/Developer/AutomationWindow/Private/SAutomationWindow.h
@@ -15,7 +15,7 @@
 #include "Widgets/Views/STableRow.h"
 #include "Widgets/Input/SComboBox.h"
 #include "SAutomationGraphicalResultBox.h"
-#include "SAutomationTestTreeView.h"
+#include "Widgets/Views/STreeView.h"
 
 #if WITH_EDITOR
 #include "AssetRegistry/IAssetRegistry.h"
@@ -114,7 +114,7 @@ public:
 	static TArray<FString> SaveExpandedTestNames(TSet<TSharedPtr<IAutomationReport>> ExpandedItems);
 
 	// Expanded the given item if its name is in the array of strings given.
-	static void ExpandItemsInList(TSharedPtr<SAutomationTestTreeView<TSharedPtr<IAutomationReport>>> InTestTable, TSharedPtr<IAutomationReport> InReport, TArray<FString> ItemsToExpand);
+	static void ExpandItemsInList(TSharedPtr<STreeView<TSharedPtr<IAutomationReport>>> InTestTable, TSharedPtr<IAutomationReport> InReport, TArray<FString> ItemsToExpand);
 
 protected:
 
@@ -569,11 +569,8 @@ private:
 	/** Must maintain a widget size so the header and row icons can line up. */
 	const float ColumnWidth;
 
-	/** Global checkbox to enable/disable all visible tests. */
-	TSharedPtr< SCheckBox > HeaderCheckbox;
-
 	/** The list of all valid tests. */
-	TSharedPtr< SAutomationTestTreeView< TSharedPtr< IAutomationReport > > > TestTable;
+	TSharedPtr< STreeView< TSharedPtr< IAutomationReport > > > TestTable;
 	
 	/** Widget for header platform icons. */
 	TSharedPtr< SHorizontalBox > PlatformsHBox;
@@ -649,9 +646,6 @@ private:
 	/** Holds a pointer to the preset text box. */
 	TSharedPtr<SEditableTextBox> PresetTextBox;
 
-	/** Hold a pointer to the test tables header row. */
-	TSharedPtr<SHeaderRow> TestTableHeaderRow;
-
 	/** List of expanded test items to preserve during a refresh */
 	TArray<FString> SavedExpandedItems;
 
