Fix play level test search mechanic
* get the package object to get the package path instead of guessing from file name

#rnx

diff --git a/Engine/Source/Developer/AutomationWindow/Private/SAutomationWindow.cpp b/Engine/Source/Developer/AutomationWindow/Private/SAutomationWindow.cpp
--- a/Engine/Source/Developer/AutomationWindow/Private/SAutomationWindow.cpp
+++ b/Engine/Source/Developer/AutomationWindow/Private/SAutomationWindow.cpp
@@ -1383,17 +1383,6 @@ void SAutomationWindow::RunSelectedTests()
 
 namespace
 {
-	bool MakeMapPathUrl(FString& InPath)
-	{
-		if ( FPaths::MakePathRelativeTo(InPath, *FPaths::ProjectContentDir()) )
-		{
-			InPath.InsertAt(0, TEXT("/Game/"));
-			InPath.RemoveFromEnd(TEXT(".umap"));
-			return true;
-		}
-		return false;
-	}
-
 	/**
 	 * Kind of a hack - this requires that we know we group all the map tests coming from blueprints under "Functional Tests"
 	 */
@@ -1437,15 +1426,15 @@ namespace
 void SAutomationWindow::FindTestReportsForCurrentEditorLevel(TArray<TSharedPtr<IAutomationReport>>& OutLevelReports)
 {
 	// Find the current map path
-	if ( GWorld && GWorld->GetCurrentLevel() )
+	if ( GWorld && GWorld->GetCurrentLevel() && GWorld->GetCurrentLevel()->GetPackage())
 	{
-		FString MapUrl(FEditorFileUtils::GetFilename(GWorld->GetCurrentLevel()));
-		if ( MakeMapPathUrl(MapUrl) )
+		const FString MapPath = GWorld->GetCurrentLevel()->GetPackage()->GetPathName();
+		if (!MapPath.IsEmpty())
 		{
 			auto FunctionTestsReport = GetFunctionalTestsReport(AutomationController->GetReports());
-			if ( FunctionTestsReport.IsValid() )
+			if (FunctionTestsReport.IsValid())
 			{
-				FindReportByGameRelativeAssetPath(FunctionTestsReport, MapUrl, OutLevelReports);
+				FindReportByGameRelativeAssetPath(FunctionTestsReport, MapPath, OutLevelReports);
 			}
 		}
 	}
