Added #if WITH_EDITOR for "Keep Pie Open option is visible" widget due to it is Editor only feature

diff --git a/Engine/Source/Developer/AutomationWindow/Private/SAutomationWindow.cpp b/Engine/Source/Developer/AutomationWindow/Private/SAutomationWindow.cpp
--- a/Engine/Source/Developer/AutomationWindow/Private/SAutomationWindow.cpp
+++ b/Engine/Source/Developer/AutomationWindow/Private/SAutomationWindow.cpp
@@ -1290,7 +1290,9 @@ TSharedRef< SWidget > SAutomationWindow::GenerateTestsOptionsMenuContent( )
 	{
 		MenuBuilder.AddWidget(NumTests, FText::GetEmpty());
 		MenuBuilder.AddWidget(SendAnalyticsWidget, FText::GetEmpty());
+		#if WITH_EDITOR
 		MenuBuilder.AddWidget(KeepPIEOpenWidget, FText::GetEmpty());
+		#endif //WITH_EDITOR
 	}
 	MenuBuilder.EndSection();
 
