Update Test Automation UI/UX

diff --git a/Engine/Plugins/Experimental/ChaosEditor/Source/FractureEditor/Private/FractureEditorStyle.cpp b/Engine/Plugins/Experimental/ChaosEditor/Source/FractureEditor/Private/FractureEditorStyle.cpp
--- a/Engine/Plugins/Experimental/ChaosEditor/Source/FractureEditor/Private/FractureEditorStyle.cpp
+++ b/Engine/Plugins/Experimental/ChaosEditor/Source/FractureEditor/Private/FractureEditorStyle.cpp
@@ -67,7 +67,7 @@ FFractureEditorStyle::FFractureEditorStyle()
 	Set("FractureEditor.Resample",             new IMAGE_BRUSH_SVG("FractureResample", IconSize));
 
 	Set("FractureEditor.ToMesh",        	   new CORE_IMAGE_BRUSH_SVG("../Editor/Slate/Starship/Common/MakeStaticMesh", IconSize));
-	Set("FractureEditor.Validate",        	   new CORE_IMAGE_BRUSH_SVG("../Editor/Slate/Starship/Common/Test", IconSize));
+	Set("FractureEditor.Validate",        	   new CORE_IMAGE_BRUSH_SVG("Starship/Common/Test", IconSize));
 
 	Set("FractureEditor.Convex",        	   new IMAGE_BRUSH_SVG("FractureConvex", IconSize));
 	Set("FractureEditor.CustomVoronoi",        new IMAGE_BRUSH_SVG("FractureCustom", IconSize));
diff --git a/Engine/Source/Developer/AutomationWindow/Private/AutomationWindowStyle.cpp b/Engine/Source/Developer/AutomationWindow/Private/AutomationWindowStyle.cpp
--- a/Engine/Source/Developer/AutomationWindow/Private/AutomationWindowStyle.cpp
+++ b/Engine/Source/Developer/AutomationWindow/Private/AutomationWindowStyle.cpp
@@ -40,7 +40,8 @@ FAutomationWindowStyle::FAutomationWindowStyle()
 
 	SetParentStyleName(FAppStyle::GetAppStyleSetName());
 
-	SetContentRoot(FPaths::EngineContentDir() / TEXT("Slate"));
+	SetContentRoot(FPaths::EngineContentDir() / TEXT("Editor/Slate"));
+	SetCoreContentRoot(FPaths::EngineContentDir() / TEXT("Slate"));
 
 	{
 		FTextBlockStyle NormalText = FAppStyle::GetWidgetStyle<FTextBlockStyle>("NormalText");
@@ -66,48 +67,61 @@ FAutomationWindowStyle::FAutomationWindowStyle()
 			.SetColorAndOpacity(FStyleColors::ForegroundHeader));
 
 		//state of individual tests
-		Set("Automation.Success", new IMAGE_BRUSH("Automation/Success", Icon16x16));
-		Set("Automation.Warning", new IMAGE_BRUSH("Automation/Warning", Icon16x16));
-		Set("Automation.Fail", new IMAGE_BRUSH("Automation/Fail", Icon16x16));
-		Set("Automation.InProcess", new IMAGE_BRUSH("Automation/InProcess", Icon16x16));
-		Set("Automation.NotRun", new IMAGE_BRUSH("Automation/NotRun", Icon16x16, FLinearColor(0.0f, 0.0f, 0.0f, 0.4f)));
-		Set("Automation.Skipped", new IMAGE_BRUSH("Automation/NoSessionWarning", Icon16x16));
-		Set("Automation.ParticipantsWarning", new IMAGE_BRUSH("Automation/ParticipantsWarning", Icon16x16));
-		Set("Automation.Participant", new IMAGE_BRUSH("Automation/Participant", Icon16x16));
+		Set("Automation.Success", new CORE_IMAGE_BRUSH("Automation/Success", Icon16x16));
+		Set("Automation.Warning", new CORE_IMAGE_BRUSH("Automation/Warning", Icon16x16));
+		Set("Automation.Fail", new CORE_IMAGE_BRUSH("Automation/Fail", Icon16x16));
+		Set("Automation.InProcess", new CORE_IMAGE_BRUSH("Automation/InProcess", Icon16x16));
+		Set("Automation.NotRun", new CORE_IMAGE_BRUSH("Automation/NotRun", Icon16x16, FLinearColor(0.0f, 0.0f, 0.0f, 0.4f)));
+		Set("Automation.Skipped", new CORE_IMAGE_BRUSH("Automation/NoSessionWarning", Icon16x16));
+		Set("Automation.ParticipantsWarning", new CORE_IMAGE_BRUSH_SVG("Starship/Common/PlayerController", Icon16x16));
+		Set("Automation.Participant", new CORE_IMAGE_BRUSH_SVG("Starship/Common/PlayerController", Icon16x16));
 
 		//status as a regression test or not
-		Set("Automation.SmokeTest", new IMAGE_BRUSH("Automation/SmokeTest", Icon16x16));
-		Set("Automation.SmokeTestParent", new IMAGE_BRUSH("Automation/SmokeTestParent", Icon16x16));
+		Set("Automation.SmokeTest", new CORE_IMAGE_BRUSH_SVG("Starship/Common/Test", Icon16x16));
+		Set("Automation.SmokeTestParent", new CORE_IMAGE_BRUSH_SVG("Starship/Common/Test", Icon16x16));
 
 		//run icons
-		Set("AutomationWindow.RunTests", new IMAGE_BRUSH("Automation/RunTests", Icon40x40));
-		Set("AutomationWindow.RefreshTests", new IMAGE_BRUSH("Automation/RefreshTests", Icon40x40));
-		Set("AutomationWindow.FindWorkers", new IMAGE_BRUSH("Automation/RefreshWorkers", Icon40x40));
-		Set("AutomationWindow.StopTests", new IMAGE_BRUSH("Automation/StopTests", Icon40x40));
-		Set("AutomationWindow.RunTests.Small", new IMAGE_BRUSH("Automation/RunTests", Icon20x20));
-		Set("AutomationWindow.RefreshTests.Small", new IMAGE_BRUSH("Automation/RefreshTests", Icon20x20));
-		Set("AutomationWindow.FindWorkers.Small", new IMAGE_BRUSH("Automation/RefreshWorkers", Icon20x20));
-		Set("AutomationWindow.StopTests.Small", new IMAGE_BRUSH("Automation/StopTests", Icon20x20));
+		Set("AutomationWindow.RunTests", new CORE_IMAGE_BRUSH_SVG("Starship/Common/play", Icon20x20));
+		Set("AutomationWindow.RunLevelTest", new IMAGE_BRUSH_SVG("Starship/Common/RunLevelTest", Icon20x20));
+		Set("AutomationWindow.RefreshTests", new CORE_IMAGE_BRUSH_SVG("Starship/Common/Update", Icon20x20));
+		Set("AutomationWindow.FindWorkers", new CORE_IMAGE_BRUSH_SVG("Starship/Common/Search_20", Icon20x20));
+		Set("AutomationWindow.StopTests", new CORE_IMAGE_BRUSH_SVG("Starship/Common/stop", Icon20x20));
+		Set("AutomationWindow.RunTests.Small", new CORE_IMAGE_BRUSH_SVG("Starship/Common/play", Icon20x20));
+		Set("AutomationWindow.RunLevelTest.Small", new IMAGE_BRUSH_SVG("Starship/Common/RunLevelTest", Icon20x20));
+		Set("AutomationWindow.RefreshTests.Small", new CORE_IMAGE_BRUSH_SVG("Starship/Common/Update", Icon20x20));
+		Set("AutomationWindow.FindWorkers.Small", new CORE_IMAGE_BRUSH_SVG("Starship/Common/Search_20", Icon20x20));
+		Set("AutomationWindow.StopTests.Small", new CORE_IMAGE_BRUSH_SVG("Starship/Common/stop", Icon20x20));
 
 		//filter icons
-		Set("AutomationWindow.ErrorFilter", new IMAGE_BRUSH("Automation/ErrorFilter", Icon40x40));
-		Set("AutomationWindow.WarningFilter", new IMAGE_BRUSH("Automation/WarningFilter", Icon40x40));
-		Set("AutomationWindow.SmokeTestFilter", new IMAGE_BRUSH("Automation/SmokeTestFilter", Icon40x40));
-		Set("AutomationWindow.DeveloperDirectoryContent", new IMAGE_BRUSH("Automation/DeveloperDirectoryContent", Icon40x40));
-		Set("AutomationWindow.ExcludedTestsFilter", new IMAGE_BRUSH("Automation/ExcludedTestsFilter", Icon40x40));
-		Set("AutomationWindow.ErrorFilter.Small", new IMAGE_BRUSH("Automation/ErrorFilter", Icon20x20));
-		Set("AutomationWindow.WarningFilter.Small", new IMAGE_BRUSH("Automation/WarningFilter", Icon20x20));
-		Set("AutomationWindow.SmokeTestFilter.Small", new IMAGE_BRUSH("Automation/SmokeTestFilter", Icon20x20));
-		Set("AutomationWindow.DeveloperDirectoryContent.Small", new IMAGE_BRUSH("Automation/DeveloperDirectoryContent", Icon20x20));
-		Set("AutomationWindow.TrackHistory", new IMAGE_BRUSH("Automation/TrackTestHistory", Icon40x40));
+		Set("AutomationWindow.ErrorFilter", new CORE_IMAGE_BRUSH_SVG("Starship/Common/x-circle", Icon16x16));
+		Set("AutomationWindow.WarningFilter", new CORE_IMAGE_BRUSH_SVG("Starship/Common/alert-triangle", Icon16x16));
+		Set("AutomationWindow.SmokeTestFilter", new CORE_IMAGE_BRUSH_SVG("Starship/Common/Test", Icon16x16));
+		Set("AutomationWindow.DeveloperDirectoryContent", new CORE_IMAGE_BRUSH_SVG("Starship/Common/Developer", Icon16x16));
+		Set("AutomationWindow.ExcludedTestsFilter", new CORE_IMAGE_BRUSH_SVG("Starship/Common/reject", Icon16x16));
+		Set("AutomationWindow.ErrorFilter.Small", new CORE_IMAGE_BRUSH_SVG("Starship/Common/x-circle", Icon16x16));
+		Set("AutomationWindow.WarningFilter.Small", new CORE_IMAGE_BRUSH_SVG("Starship/Common/alert-triangle", Icon16x16));
+		Set("AutomationWindow.SmokeTestFilter.Small", new CORE_IMAGE_BRUSH_SVG("Starship/Common/Test", Icon16x16));
+		Set("AutomationWindow.DeveloperDirectoryContent.Small", new CORE_IMAGE_BRUSH_SVG("Starship/Common/Developer", Icon16x16));
 
 		//device group settings
-		Set("AutomationWindow.GroupSettings", new IMAGE_BRUSH("Automation/Groups", Icon40x40));
-		Set("AutomationWindow.GroupSettings.Small", new IMAGE_BRUSH("Automation/Groups", Icon20x20));
+		Set("AutomationWindow.GroupSettings", new CORE_IMAGE_BRUSH_SVG("Starship/Common/Group_20", Icon20x20));
+		Set("AutomationWindow.GroupSettings.Small", new CORE_IMAGE_BRUSH_SVG("Starship/Common/Group_20", Icon20x20));
 
 		//test backgrounds
-		Set("AutomationWindow.GameGroupBorder", new BOX_BRUSH("Automation/GameGroupBorder", FMargin(4.0f / 16.0f)));
-		Set("AutomationWindow.EditorGroupBorder", new BOX_BRUSH("Automation/EditorGroupBorder", FMargin(4.0f / 16.0f)));
+		Set("AutomationWindow.GameGroupBorder", new CORE_BOX_BRUSH("Automation/GameGroupBorder", FMargin(4.0f / 16.0f)));
+		Set("AutomationWindow.EditorGroupBorder", new CORE_BOX_BRUSH("Automation/EditorGroupBorder", FMargin(4.0f / 16.0f)));
+
+		FButtonStyle Button = FButtonStyle()
+			.SetNormal(BOX_BRUSH("Common/Button", FVector2D(32, 32), 8.0f / 32.0f))
+			.SetHovered(BOX_BRUSH("Common/Button_Hovered", FVector2D(32, 32), 8.0f / 32.0f))
+			.SetPressed(BOX_BRUSH("Common/Button_Pressed", FVector2D(32, 32), 8.0f / 32.0f))
+			.SetNormalPadding(FMargin(2, 2, 2, 2))
+			.SetPressedPadding(FMargin(2, 3, 2, 1));
+
+		Set("AutomationWindow.ToggleButton", FCheckBoxStyle(FStarshipCoreStyle::GetCoreStyle().GetWidgetStyle<FCheckBoxStyle>("ToggleButtonCheckbox"))
+			.SetUncheckedHoveredImage(FSlateColorBrush(FStyleColors::SelectHover))
+			.SetPadding(FMargin(16, 4))
+		);
 	}
 	
 	FSlateStyleRegistry::RegisterSlateStyle(*this);
diff --git a/Engine/Source/Developer/AutomationWindow/Private/SAutomationTestItem.cpp b/Engine/Source/Developer/AutomationWindow/Private/SAutomationTestItem.cpp
--- a/Engine/Source/Developer/AutomationWindow/Private/SAutomationTestItem.cpp
+++ b/Engine/Source/Developer/AutomationWindow/Private/SAutomationTestItem.cpp
@@ -47,7 +47,56 @@ void SAutomationTestItem::Construct( const FArguments& InArgs, const TSharedRef<
 BEGIN_SLATE_FUNCTION_BUILD_OPTIMIZATION
 TSharedRef<SWidget> SAutomationTestItem::GenerateWidgetForColumn( const FName& ColumnName )
 {
-	if( ColumnName == AutomationTestWindowConstants::Title)
+	if (ColumnName == AutomationTestWindowConstants::Checked)
+	{
+		return SNew(SHorizontalBox)
+			+ SHorizontalBox::Slot()
+			.AutoWidth()
+			.HAlign(HAlign_Center)
+			.Padding(4.0f, 0.0f)
+			[
+				//enabled/disabled check box
+				SNew(SCheckBox)
+				.IsChecked(this, &SAutomationTestItem::IsTestEnabled)
+			.OnCheckStateChanged(this, &SAutomationTestItem::HandleTestingCheckbox_Click)
+			];
+	}
+	else if (ColumnName == AutomationTestWindowConstants::Skipped)
+	{
+		return SNew(SHorizontalBox)
+			+ SHorizontalBox::Slot()
+			.HAlign(HAlign_Center)
+			.AutoWidth()
+			[
+				SNew(SButton)
+				.ButtonStyle(FAutomationWindowStyle::Get(), "NoBorder")
+				.ToolTipText(LOCTEXT("ToggleExcluded", "Exclude"))
+				.OnClicked(FOnClicked::CreateSP(this, &SAutomationTestItem::SetSkipFlag))
+				[
+					SNew(SImage)
+					.Image(FAutomationWindowStyle::Get().GetBrush("AutomationWindow.ExcludedTestsFilter"))
+					.Visibility(this, &SAutomationTestItem::IsToBeSkipped_GetVisibility)
+				]
+			];
+	}
+	else if (ColumnName == AutomationTestWindowConstants::SkippedOptions)
+	{
+		return SNew(SHorizontalBox)
+			+SHorizontalBox::Slot()
+#if WITH_EDITOR
+			.HAlign(HAlign_Center)
+			.AutoWidth()
+			[
+				SNew(SSimpleButton)
+				.Icon(FAutomationWindowStyle::Get().GetBrush("Icons.Edit"))
+				.Visibility(this, &SAutomationTestItem::IsDirectlyExcluded_GetVisibility)
+				.ToolTipText(LOCTEXT("EditExcludeOptions", "Edit exclude options"))
+				.OnClicked(FOnClicked::CreateSP(this, &SAutomationTestItem::OnEditExcludeOptionsClicked))
+			]
+#endif
+			;
+	}
+	else if( ColumnName == AutomationTestWindowConstants::Title)
 	{
 		TSharedRef<SWidget> TestNameWidget = SNullWidget::NullWidget;
 
@@ -102,16 +151,6 @@ TSharedRef<SWidget> SAutomationTestItem::GenerateWidgetForColumn( const FName& C
 		}
 
 		return SNew(SHorizontalBox)
-			+SHorizontalBox::Slot()
-			.AutoWidth()
-			.HAlign(HAlign_Center)
-			.Padding(4.0f, 0.0f)
-			[
-				//enabled/disabled check box
-				SNew( SCheckBox )
-				.IsChecked(this, &SAutomationTestItem::IsTestEnabled)
-				.OnCheckStateChanged(this, &SAutomationTestItem::HandleTestingCheckbox_Click)
-			]
 			+SHorizontalBox::Slot()
 			.AutoWidth()
 			.VAlign(VAlign_Center)
@@ -273,32 +312,6 @@ TSharedRef<SWidget> SAutomationTestItem::GenerateWidgetForColumn( const FName& C
 		return SNew( STextBlock )
 		.Text( this, &SAutomationTestItem::ItemStatus_DurationText);
 	}
-	else if (ColumnName == AutomationTestWindowConstants::IsToBeSkipped)
-	{
-		return SNew(SHorizontalBox)
-			+ SHorizontalBox::Slot()
-			.HAlign(HAlign_Center)
-			.AutoWidth()
-			[
-				SNew(SCheckBox)
-				.IsChecked(this, &SAutomationTestItem::IsToBeSkipped)
-				.IsEnabled(this, &SAutomationTestItem::IsDirectlyExcluded)
-				.OnCheckStateChanged(this, &SAutomationTestItem::SetSkipFlag)
-				.ToolTipText(this, &SAutomationTestItem::GetExcludeReason)
-#if WITH_EDITOR
-			]
-		+ SHorizontalBox::Slot()
-			.HAlign(HAlign_Center)
-			.AutoWidth()
-			[
-				SNew(SSimpleButton)
-				.Icon(FAutomationWindowStyle::Get().GetBrush("Icons.Edit"))
-				.Visibility(this, &SAutomationTestItem::IsDirectlyExcluded_GetVisibility)
-				.ToolTipText(LOCTEXT("EditExcludeOptions", "Edit exclude options"))
-				.OnClicked(FOnClicked::CreateSP(this, &SAutomationTestItem::OnEditExcludeOptionsClicked))
-#endif
-			];
-	}
 
 
 	return SNullWidget::NullWidget;
@@ -367,14 +380,14 @@ ECheckBoxState SAutomationTestItem::IsTestEnabled() const
 	return TestStatus->IsEnabled() ? ECheckBoxState::Checked : ECheckBoxState::Unchecked;
 }
 
-ECheckBoxState SAutomationTestItem::IsToBeSkipped() const
+EVisibility SAutomationTestItem::IsToBeSkipped_GetVisibility() const
 {
-	return TestStatus->IsToBeSkipped() ? ECheckBoxState::Checked : ECheckBoxState::Unchecked;
+	return TestStatus->IsToBeSkipped() ? EVisibility::Visible : EVisibility::Hidden;
 }
 
 bool SAutomationTestItem::IsDirectlyExcluded() const
 {
-	return WITH_EDITOR && IsLocalSession && !TestStatus->IsToBeSkippedByPropagation();
+	return WITH_EDITOR && IsLocalSession && TestStatus->IsToBeSkipped() && !TestStatus->IsToBeSkippedByPropagation();
 }
 
 EVisibility SAutomationTestItem::IsDirectlyExcluded_GetVisibility() const
@@ -390,10 +403,10 @@ FText SAutomationTestItem::GetExcludeReason() const
 	return IsToBeSkipped ? FText::FromName(Reason) : FText();
 }
 
-void SAutomationTestItem::SetSkipFlag(ECheckBoxState Enable)
+FReply SAutomationTestItem::SetSkipFlag()
 {
 #if WITH_EDITOR
-	if (Enable == ECheckBoxState::Checked)
+	if (!TestStatus->IsToBeSkipped())
 	{
 		OnEditExcludeOptionsClicked();
 	}
@@ -402,6 +415,7 @@ void SAutomationTestItem::SetSkipFlag(ECheckBoxState Enable)
 		TestStatus->SetSkipFlag(false);
 	}
 #endif
+	return FReply::Handled();
 }
 
 FReply SAutomationTestItem::OnEditExcludeOptionsClicked()
diff --git a/Engine/Source/Developer/AutomationWindow/Private/SAutomationTestItem.h b/Engine/Source/Developer/AutomationWindow/Private/SAutomationTestItem.h
--- a/Engine/Source/Developer/AutomationWindow/Private/SAutomationTestItem.h
+++ b/Engine/Source/Developer/AutomationWindow/Private/SAutomationTestItem.h
@@ -84,9 +84,9 @@ protected:
 	/**
 	* Is the test inside exclude list. This info is taken from Config/DefaultEngine.ini, section [AutomationTestExcludelist].
 	*
-	* @return true if the test is inside the exclude list.
+	* Return the visibility state based on skipped state.
 	*/
-	ECheckBoxState IsToBeSkipped() const;
+	EVisibility IsToBeSkipped_GetVisibility() const;
 
 	/**
 	* Is the test inside exclude list through direct exclusion (not through propagation).
@@ -110,9 +110,9 @@ protected:
 	/**
 	* Add or remove test from exclude list. It will change Config/DefaultEngine.ini, section [AutomationTestExcludelist].
 	*
-	* @param if true, add this item to exclude list, remove otherwise.
+	* @return respond of dialog
 	*/
-	void SetSkipFlag(ECheckBoxState Enable);
+	FReply SetSkipFlag();
 
 	/**
 	* Open Exclude test options editor
diff --git a/Engine/Source/Developer/AutomationWindow/Private/SAutomationWindow.cpp b/Engine/Source/Developer/AutomationWindow/Private/SAutomationWindow.cpp
--- a/Engine/Source/Developer/AutomationWindow/Private/SAutomationWindow.cpp
+++ b/Engine/Source/Developer/AutomationWindow/Private/SAutomationWindow.cpp
@@ -28,6 +28,7 @@
 #include "AutomationPresetManager.h"
 #include "SAutomationTestItemContextMenu.h"
 #include "SAutomationTestItem.h"
+#include "SSimpleComboButton.h"
 
 #if WITH_EDITOR
 	#include "Engine/Level.h"
@@ -76,6 +77,10 @@ public:
 		// Added button for running the currently open level test.
 		UI_COMMAND(RunLevelTest, "Run Level Test", "Run Level Test", EUserInterfaceActionType::Button, FInputChord());
 #endif
+
+		UI_COMMAND(CreateNewPreset, "Create Automation Preset", "Create New Preset", EUserInterfaceActionType::Button, FInputChord());
+		UI_COMMAND(RemoveCurrentPreset, "Delete", "Remove Current Preset", EUserInterfaceActionType::Button, FInputChord());
+		UI_COMMAND(RenameCurrentPreset, "Rename", "Rename Current Preset", EUserInterfaceActionType::Button, FInputChord());
 	}
 public:
 	TSharedPtr<FUICommandInfo> RefreshTests;
@@ -88,15 +93,19 @@ public:
 #if WITH_EDITOR
 	TSharedPtr<FUICommandInfo> RunLevelTest;
 #endif
+
+	TSharedPtr<FUICommandInfo> CreateNewPreset;
+	TSharedPtr<FUICommandInfo> RemoveCurrentPreset;
+	TSharedPtr<FUICommandInfo> RenameCurrentPreset;
 };
 
 //////////////////////////////////////////////////////////////////////////
 // SAutomationWindow
 
 SAutomationWindow::SAutomationWindow() 
-	: ColumnWidth(50.0f)
+	: ColumnWidth(90.0f)
+	, bIsLabelVisibilityEnabled(false)
 {
-
 }
 
 SAutomationWindow::~SAutomationWindow()
@@ -155,6 +164,7 @@ void SAutomationWindow::Construct( const FArguments& InArgs, const IAutomationCo
 	TestPresetManager = MakeShareable(new FAutomationTestPresetManager());
 	TestPresetManager->LoadPresets();
 	bAddingTestPreset = false;
+	bRenamingTestPreset = false;
 
 	bHasChildTestSelected = false;
 
@@ -205,30 +215,59 @@ void SAutomationWindow::Construct( const FArguments& InArgs, const IAutomationCo
 		.HeaderRow
 		(
 			SNew(SHeaderRow)
+			+ SHeaderRow::Column( AutomationTestWindowConstants::Checked )
+			.FixedWidth(30.0f)
+			.HAlignHeader(HAlign_Center)
+			.VAlignHeader(VAlign_Center)
+			.HAlignCell(HAlign_Center)
+			.VAlignCell(VAlign_Center)
+			[
+				//global enable/disable check box
+				SNew(SCheckBox)
+				.OnCheckStateChanged(this, &SAutomationWindow::HeaderCheckboxStateChange)
+				.ToolTipText(LOCTEXT("Enable Disable Test", "Enable / Disable Test"))
+			]
+
+			+ SHeaderRow::Column( AutomationTestWindowConstants::Skipped )
+			.FixedWidth(30.0f)
+			.HAlignHeader(HAlign_Center)
+			.VAlignHeader(VAlign_Center)
+			.HAlignCell(HAlign_Center)
+			.VAlignCell(VAlign_Center)
+			[
+				SNew(SImage)
+				.ColorAndOpacity(FLinearColor(1.0f, 1.0f, 1.0f, 0.4f))
+				.ToolTipText(LOCTEXT("Excluded", "Excluded"))
+				.Image(FAutomationWindowStyle::Get().GetBrush("AutomationWindow.ExcludedTestsFilter"))
+			]
+
+			+ SHeaderRow::Column( AutomationTestWindowConstants::SkippedOptions )
+			.FixedWidth(30.0f)
+			.HAlignHeader(HAlign_Center)
+			.VAlignHeader(VAlign_Center)
+			.HAlignCell(HAlign_Center)
+			.VAlignCell(VAlign_Center)
+			[
+				SNew(SImage)
+				.ColorAndOpacity(FLinearColor(1.0f, 1.0f, 1.0f, 0.4f))
+				.Image(FAutomationWindowStyle::Get().GetBrush("Icons.Edit"))
+			]
+
 			+ SHeaderRow::Column( AutomationTestWindowConstants::Title )
 			.FillWidth(0.80f)
 			[
 				SNew(SHorizontalBox)
 				+SHorizontalBox::Slot()
 				.AutoWidth()
-				.HAlign(HAlign_Center)
-				[
-					//global enable/disable check box
-					SNew(SCheckBox)
-					.OnCheckStateChanged( this, &SAutomationWindow::HeaderCheckboxStateChange)
-					.ToolTipText( LOCTEXT( "Enable Disable Test", "Enable / Disable  Test" ) )
-				]
-				+SHorizontalBox::Slot()
-				.AutoWidth()
 				.VAlign(VAlign_Center)
 				[
 					SNew( STextBlock )
-					.Text( LOCTEXT("TestName_Header", "Test Name") )
+					.Text( LOCTEXT("TestName_Header", "Test") )
 				]
 			]
 	
 			+ SHeaderRow::Column( AutomationTestWindowConstants::SmokeTest )
-			.FixedWidth( 50.0f )
+			.FixedWidth( 30.0f )
 			.HAlignHeader(HAlign_Center)
 			.VAlignHeader(VAlign_Center)
 			.HAlignCell(HAlign_Center)
@@ -242,7 +281,7 @@ void SAutomationWindow::Construct( const FArguments& InArgs, const IAutomationCo
 			]
 	
 			+ SHeaderRow::Column( AutomationTestWindowConstants::RequiredDeviceCount )
-			.FixedWidth(50.0f)
+			.FixedWidth(30.0f)
 			.HAlignHeader(HAlign_Center)
 			.VAlignHeader(VAlign_Center)
 			.HAlignCell(HAlign_Center)
@@ -254,20 +293,17 @@ void SAutomationWindow::Construct( const FArguments& InArgs, const IAutomationCo
 			]
 	
 			+ SHeaderRow::Column(AutomationTestWindowConstants::Timing)
-			.FixedWidth(100.0f)
+			.FixedWidth( 80.0f )
 			.DefaultLabel(LOCTEXT("TestDurationRange", "Duration"))
 
 			+ SHeaderRow::Column( AutomationTestWindowConstants::Status )
-			.FixedWidth(50.0f)
+			.FixedWidth(ColumnWidth)
+			.HAlignHeader(HAlign_Left)
+			.VAlignHeader(VAlign_Center)
 			[
 				//platform header placeholder
 				PlatformsHBox.ToSharedRef()
 			]
-
-			+ SHeaderRow::Column(AutomationTestWindowConstants::IsToBeSkipped)
-			.FillWidth(0.10f)
-			.DefaultLabel(LOCTEXT("Excluded", "Excluded"))
-
 		);
 
 	RequestedFilterComboList.Empty();
@@ -299,7 +335,6 @@ void SAutomationWindow::Construct( const FArguments& InArgs, const IAutomationCo
 	this->ChildSlot
 	[
 		SNew(SOverlay)
-
 		+ SOverlay::Slot()
 			[
 				SNew( SSplitter )
@@ -319,9 +354,9 @@ void SAutomationWindow::Construct( const FArguments& InArgs, const IAutomationCo
 						SNew( SHorizontalBox )
 
 						+ SHorizontalBox::Slot()
-						.AutoWidth()
 						.HAlign(HAlign_Left)
 						.VAlign(VAlign_Center)
+						.FillWidth(1)
 						[
 							SAutomationWindow::MakeAutomationWindowToolBar( AutomationWindowActions.ToSharedRef(), SharedThis(this) )
 						]
@@ -329,13 +364,13 @@ void SAutomationWindow::Construct( const FArguments& InArgs, const IAutomationCo
 
 					+ SVerticalBox::Slot()
 					.FillHeight(1.0f)
-					.Padding(0.0f, 4.0f, 0.0f, 0.0f)
+					.Padding(0.0f)
 					[
 						SNew(SOverlay)
 						+ SOverlay::Slot()
 						[
 							SNew(SBorder)
-							.BorderImage(this, &SAutomationWindow::GetTestBackgroundBorderImage)
+							.BorderImage(FAutomationWindowStyle::Get().GetBrush("NoBorder"))
 							.Padding(3)
 							[
 								SNew(SBox)
@@ -369,6 +404,14 @@ void SAutomationWindow::Construct( const FArguments& InArgs, const IAutomationCo
 											]
 										]
 
+
+										+ SHorizontalBox::Slot()
+										.FillWidth(0.02f)
+										[
+											SNew(SSpacer)
+											.Visibility(this, &SAutomationWindow::HandleGroupsVisibility)
+										]
+
 										+ SHorizontalBox::Slot()
 										.AutoWidth()
 										.VAlign(VAlign_Center)
@@ -391,7 +434,13 @@ void SAutomationWindow::Construct( const FArguments& InArgs, const IAutomationCo
 										]
 
 										+ SHorizontalBox::Slot()
-										.FillWidth(1.0f)
+										.FillWidth(0.02f)
+										[
+											SNew(SSpacer)
+										]
+
+										+ SHorizontalBox::Slot()
+										.FillWidth(0.8f)
 										.VAlign(VAlign_Center)
 										.Padding(2.0f, 0, 0, 0)
 										[
@@ -400,6 +449,28 @@ void SAutomationWindow::Construct( const FArguments& InArgs, const IAutomationCo
 											.OnTextChanged(this, &SAutomationWindow::OnFilterTextChanged)
 											.IsEnabled(this, &SAutomationWindow::IsAutomationControllerIdle)
 										]
+
+										+ SHorizontalBox::Slot()
+										.FillWidth(0.02f)
+										[
+											SNew(SSpacer)
+										]
+
+										+SHorizontalBox::Slot()
+										.VAlign(VAlign_Center)
+										.AutoWidth()
+										[
+											SNew(STextBlock)
+											.Text( LOCTEXT("AutomationFilterLabel", "Filter:") )
+										]
+
+										+SHorizontalBox::Slot()
+										.AutoWidth()
+										.HAlign(HAlign_Left)
+										.VAlign(VAlign_Center)
+										[
+											SAutomationWindow::MakeAutomationFilterToolBar(AutomationWindowActions.ToSharedRef(), SharedThis(this))
+										]
 									]
 
 									+ SVerticalBox::Slot()
@@ -667,6 +738,21 @@ void SAutomationWindow::CreateCommands()
 		FCanExecuteAction::CreateRaw(this, &SAutomationWindow::CanExecuteRunLevelTest)
 		);
 #endif // WITH_EDITOR
+
+	ActionList.MapAction(Commands.CreateNewPreset,
+		FExecuteAction::CreateRaw(this, &SAutomationWindow::OnNewPresetClicked),
+		FCanExecuteAction::CreateRaw(this, &SAutomationWindow::IsAutomationControllerIdle)
+	);
+
+	ActionList.MapAction(Commands.RemoveCurrentPreset,
+		FExecuteAction::CreateRaw(this, &SAutomationWindow::OnRemovePresetClicked),
+		FCanExecuteAction::CreateRaw(this, &SAutomationWindow::IsAutomationControllerIdle)
+	);
+
+	ActionList.MapAction(Commands.RenameCurrentPreset,
+		FExecuteAction::CreateRaw(this, &SAutomationWindow::OnRenamePresetClicked),
+		FCanExecuteAction::CreateRaw(this, &SAutomationWindow::IsAutomationControllerIdle)
+	);
 }
 
 TSharedRef< SWidget > SAutomationWindow::MakeAutomationWindowToolBar( const TSharedRef<FUICommandList>& InCommandList, TSharedPtr<class SAutomationWindow> InAutomationWindow )
@@ -679,19 +765,11 @@ TSharedRef< SWidget > SAutomationWindow::MakeAutomationWindowToolBar( const TSha
 {
 	struct Local
 	{
-		static void FillToolbar(FToolBarBuilder& ToolbarBuilder, TSharedRef<SWidget> RunTests, TSharedRef<SWidget> PresetBox, TWeakPtr<class SAutomationWindow> InAutomationWindow)
+		static void FillToolbar(FToolBarBuilder& ToolbarBuilder, TSharedRef<SWidget> RunTests, TSharedRef<SWidget> DeviceGroups, TSharedRef<SWidget> PresetBox, TWeakPtr<class SAutomationWindow> InAutomationWindow)
 		{
 			ToolbarBuilder.BeginSection("Automation");
 			{
-				ToolbarBuilder.AddWidget( RunTests );
-				FUIAction DefaultAction;
-				ToolbarBuilder.AddComboButton(
-					DefaultAction,
-					FOnGetContent::CreateStatic( &SAutomationWindow::GenerateTestsOptionsMenuContent, InAutomationWindow ),
-					TAttribute<FText>(),
-					LOCTEXT( "TestOptionsToolTip", "Test Options" ),
-					FSlateIcon(FAutomationWindowStyle::Get().GetStyleSetName(), "AutomationWindow.TestOptions"),
-					false);
+				ToolbarBuilder.AddWidget(RunTests);
 
 				// Added button for running the currently open level test.
 #if WITH_EDITOR
@@ -700,159 +778,174 @@ TSharedRef< SWidget > SAutomationWindow::MakeAutomationWindowToolBar( const TSha
 					NAME_None,
 					TAttribute<FText>(),
 					LOCTEXT("RunLevelTest_ToolTip", "If the currently loaded editor level is a test map, click this to select the test and run it immediately."),
-					FSlateIcon(FAutomationWindowStyle::Get().GetStyleSetName(), "AutomationWindow.RunTests"));
+					FSlateIcon(FAutomationWindowStyle::Get().GetStyleSetName(), "AutomationWindow.RunLevelTest"));
 #endif
 
+				ToolbarBuilder.AddWidget(SNew(SSpacer));
 				ToolbarBuilder.AddToolBarButton( FAutomationWindowCommands::Get().RefreshTests );
+				ToolbarBuilder.AddWidget(SNew(SSpacer));
 				ToolbarBuilder.AddToolBarButton( FAutomationWindowCommands::Get().FindWorkers );
-			}
-			ToolbarBuilder.EndSection();
-			ToolbarBuilder.BeginSection("Filters");
-			{
-				ToolbarBuilder.AddToolBarButton( FAutomationWindowCommands::Get().ErrorFilter );
-				ToolbarBuilder.AddToolBarButton( FAutomationWindowCommands::Get().WarningFilter );
-				ToolbarBuilder.AddToolBarButton( FAutomationWindowCommands::Get().DeveloperDirectoryContent );
-				ToolbarBuilder.AddToolBarButton( FAutomationWindowCommands::Get().ExcludedTestsFilter);
+				ToolbarBuilder.AddWidget(SNew(SSpacer));
 			}
 			ToolbarBuilder.EndSection();
 			ToolbarBuilder.BeginSection("GroupFlags");
 			{
-				ToolbarBuilder.AddComboButton(
-					FUIAction(),
-					FOnGetContent::CreateStatic( &SAutomationWindow::GenerateGroupOptionsMenuContent, InAutomationWindow ),
-					LOCTEXT( "GroupOptions_Label", "Device Groups" ),
-					LOCTEXT( "GroupOptionsToolTip", "Device Group Options" ),
-					FSlateIcon(FAutomationWindowStyle::Get().GetStyleSetName(), "AutomationWindow.GroupSettings"));
+				ToolbarBuilder.AddWidget(SNew(SSpacer));
+				ToolbarBuilder.AddWidget(DeviceGroups);
+				ToolbarBuilder.AddWidget(SNew(SSpacer));
 			}
 			ToolbarBuilder.EndSection();
 			ToolbarBuilder.BeginSection("Presets");
 			{
+				ToolbarBuilder.AddWidget(SNew(SSpacer));
 				ToolbarBuilder.AddWidget( PresetBox );
+				ToolbarBuilder.AddWidget(SNew(SSpacer));
 			}
 			ToolbarBuilder.EndSection();
 		}
 	};
 
 	TSharedRef<SWidget> RunTests = 
-		SNew( SButton )
-		.ToolTipText( LOCTEXT( "StartStop Tests", "Start / Stop tests" ) )
-		.OnClicked( this, &SAutomationWindow::RunTests )
-		.IsEnabled( this, &SAutomationWindow::IsAutomationRunButtonEnabled )
-		.ContentPadding(0)
+		SNew(SHorizontalBox)
+		+ SHorizontalBox::Slot()
+		.VAlign( VAlign_Center )
+		.AutoWidth()
+		.Padding(0)
 		[
-			SNew(SHorizontalBox)
-			+ SHorizontalBox::Slot()
-			.FillWidth(1)
-			.VAlign(VAlign_Center)
+			SNew( SComboButton )
+			.ComboButtonStyle( FAppStyle::Get(), "SimpleComboButton" )
+			.OnGetMenuContent( this, &SAutomationWindow::GenerateTestsOptionsMenuContent )
+			.ToolTipText( LOCTEXT("TestOptionsToolTip", "Test Options") )
+			.HasDownArrow(true)
+			.ContentPadding(0)
+			.ButtonContent()
 			[
-				SNew( SVerticalBox )
-				+ SVerticalBox::Slot()
-				.AutoHeight()
-				.HAlign( HAlign_Center )
+
+			SNew( SButton )
+			.ButtonStyle(FAppStyle::Get(), "NoBorder")
+			.ToolTipText( LOCTEXT( "StartStop Tests", "Start / Stop tests" ) )
+			.OnClicked( this, &SAutomationWindow::RunTests )
+			.IsEnabled( this, &SAutomationWindow::IsAutomationRunButtonEnabled )
+			.HAlign( HAlign_Center )
+			.VAlign( VAlign_Center )
+			.ContentPadding(0)
+			[
+				SNew( SHorizontalBox )
+				+ SHorizontalBox::Slot()
+				.VAlign( VAlign_Center )
 				[
-					SNew( SOverlay )
-					+SOverlay::Slot()
+					SNew( SVerticalBox )
+					+ SVerticalBox::Slot()
+					.AutoHeight()
 					[
-						SNew( SImage ) 
-						.Image( this, &SAutomationWindow::GetRunAutomationIcon )
+						SNew( SOverlay )
+						+SOverlay::Slot()
+						.Padding(0)
+						.VAlign( VAlign_Center )
+						.HAlign( HAlign_Center )
+						[
+							SNew( SImage ) 
+							.Image( this, &SAutomationWindow::GetRunAutomationIcon )
+						]
+						+SOverlay::Slot()
+						.Padding(0)
+						.VAlign( VAlign_Bottom )
+						.HAlign( HAlign_Right )
+						[
+							SNew( STextBlock )
+							.Margin( FMargin(22.0f, 25.0f, 0.0, 0.0f) )
+							.Text( this, &SAutomationWindow::OnGetNumEnabledTestsString )
+							.ColorAndOpacity(FSlateColor::UseForeground())
+							.MinDesiredWidth(55.0f)
+							.Justification(ETextJustify::Right)
+						]
 					]
-					+SOverlay::Slot()
-					.HAlign(HAlign_Right)
-					.VAlign(VAlign_Bottom)
+					+ SVerticalBox::Slot()
+					.AutoHeight()
+					.HAlign( HAlign_Center )
 					[
 						SNew( STextBlock )
-						.Text( this, &SAutomationWindow::OnGetNumEnabledTestsString )
+						.Visibility( bIsLabelVisibilityEnabled ? EVisibility::Visible : EVisibility::Collapsed )
+						.Text( this, &SAutomationWindow::GetRunAutomationLabel )
 						.ColorAndOpacity(FSlateColor::UseForeground())
 					]
 				]
-				+ SVerticalBox::Slot()
-				.AutoHeight()
-				.HAlign( HAlign_Center )
-				[
-					SNew( STextBlock )
-					.Visibility( this, &SAutomationWindow::GetLargeToolBarVisibility )
-					.Text( this, &SAutomationWindow::GetRunAutomationLabel )
-					.ColorAndOpacity(FSlateColor::UseForeground())
-				]
+			]
+
 			]
 		];
 
+	TSharedRef<SWidget> DeviceGroups =
+		SNew(SVerticalBox)
+		+ SVerticalBox::Slot()
+		.AutoHeight()
+		.VAlign(VAlign_Center)
+		.FillHeight(1)
+		[
+			SNew(SSimpleComboButton)
+			.OnGetMenuContent(this, &SAutomationWindow::GenerateGroupOptionsMenuContent)
+			.Icon(FAutomationWindowStyle::Get().GetBrush(FName("AutomationWindow.GroupSettings" + GetSmallIconExtension())))
+			.Text(LOCTEXT("GroupOptions_Label", "Device Groups"))
+			.ToolTipText(LOCTEXT("GroupOptionsToolTip", "Device Group Options"))
+			.HasDownArrow(true)
+		];
+
 	TSharedRef<SWidget> TestPresets = 
 		SNew( SVerticalBox )
 		+SVerticalBox::Slot()
-		.FillHeight(0.75f)
-		.VAlign(VAlign_Bottom)
-		.HAlign(HAlign_Left)
-		[
-			SNew( SHorizontalBox )
-			+SHorizontalBox::Slot()
-			.AutoWidth()
-			.Padding(2.0f)
-			.VAlign(VAlign_Bottom)
-			[
-				SNew( STextBlock )
-				.Text( LOCTEXT("AutomationPresetLabel", "Preset:") )
-				.IsEnabled( this, &SAutomationWindow::IsAutomationControllerIdle )
-			]
-		]
-		+SVerticalBox::Slot()
-		.AutoHeight()
-		.VAlign(VAlign_Bottom)
+		.FillHeight(1)
+		.VAlign( VAlign_Center )
+		.Padding( 5.0f )
 		[
 			SNew( SHorizontalBox )
 			+SHorizontalBox::Slot()
 			.AutoWidth()
 			[
+				SNew( SOverlay )
 				//Preset Combo / Text
-				SNew(SOverlay)
 				+SOverlay::Slot()
 				[
-					SNew(SHorizontalBox)
-					.Visibility(this,&SAutomationWindow::HandlePresetComboVisibility)
+					SNew( SHorizontalBox )
+					+ SHorizontalBox::Slot()
+					.VAlign( VAlign_Center )
+					.FillWidth( 65.0f )
+					[
+						SNew( STextBlock )
+						.Text( LOCTEXT("AutomationPresetLabel", "Preset:") )
+					]
 					+SHorizontalBox::Slot()
-					.FillWidth(1.0)
+					.AutoWidth()
 					[
-						SAssignNew( PresetComboBox, SComboBox< TSharedPtr<FAutomationTestPreset> > )
-						.OptionsSource( &TestPresetManager->GetAllPresets() )
-						.OnGenerateWidget( this, &SAutomationWindow::GeneratePresetComboItem )
-						.OnSelectionChanged( this, &SAutomationWindow::HandlePresetChanged )
-						.IsEnabled(this, &SAutomationWindow::IsAutomationControllerIdle)
+						SNew( SComboButton )
+						.OnGetMenuContent( this, &SAutomationWindow::GeneratePresetsMenuContent )
+						.ButtonContent()
 						[
-							SNew( STextBlock )
-							.Text( this, &SAutomationWindow::GetPresetComboText )
+							SNew( SHorizontalBox )
+							+ SHorizontalBox::Slot()
+							.FillWidth(1)
+							.VAlign( VAlign_Center )
+							[
+								SAssignNew( PresetComboButtonText, STextBlock )
+								.MinDesiredWidth( 60.0f )
+								.Text( LOCTEXT("AutomationPreset_None", "None") )
+							]
 						]
 					]
 				]
 				+SOverlay::Slot()
 				[
-					SNew(SHorizontalBox)
-					.Visibility(this,&SAutomationWindow::HandlePresetTextVisibility)
+					SNew( SHorizontalBox )
+					.Visibility( this, &SAutomationWindow::HandlePresetTextVisibility )
 					+SHorizontalBox::Slot()
-					.FillWidth(1.0)
+					.FillWidth( 100.0f )
 					[
-						SAssignNew(PresetTextBox, SEditableTextBox)
-						.OnTextCommitted(this, &SAutomationWindow::HandlePresetTextCommited)
-						.IsEnabled(this, &SAutomationWindow::IsAutomationControllerIdle)
+						SAssignNew( PresetTextBox, SEditableTextBox )
+						.OnTextCommitted( this, &SAutomationWindow::HandlePresetTextCommited )
+						.IsEnabled( this, &SAutomationWindow::IsAutomationControllerIdle )
 					]
 				]
 			]
 
-			//New button
-			+SHorizontalBox::Slot()
-			.AutoWidth()
-			[
-				SNew( SButton )
-				.ButtonStyle(FAutomationWindowStyle::Get(), "NoBorder" )
-				.OnClicked( this, &SAutomationWindow::HandleNewPresetClicked )
-				.ToolTipText( LOCTEXT("AutomationPresetNewButtonTooltip", "Create a new preset") )
-				.IsEnabled(this, &SAutomationWindow::IsAddButtonEnabled)
-				.Content()
-				[
-					SNew(SImage)
-					.Image(FAutomationWindowStyle::Get().GetBrush("Icons.PlusCircle"))
-				]
-			]
-
 			//Save button
 			+SHorizontalBox::Slot()
 			.AutoWidth()
@@ -864,32 +957,16 @@ TSharedRef< SWidget > SAutomationWindow::MakeAutomationWindowToolBar( const TSha
 				.IsEnabled(this, &SAutomationWindow::IsSaveButtonEnabled)
 				.Content()
 				[
-					SNew(SImage)
+					SNew( SImage )
 					.Image(FAutomationWindowStyle::Get().GetBrush("Icons.Save"))
 				]
 			]
-
-			//Remove button
-			+ SHorizontalBox::Slot()
-			.AutoWidth()
-			[
-				// remove button
-				SNew(SButton)
-				.ButtonStyle(FAutomationWindowStyle::Get(), "NoBorder" )
-				.OnClicked( this, &SAutomationWindow::HandleRemovePresetClicked )
-				.ToolTipText(LOCTEXT("AutomationPresetRemoveButtonTooltip", "Remove the selected preset"))
-				.IsEnabled(this, &SAutomationWindow::IsRemoveButtonEnabled)
-				.Content()
-				[
-					SNew(SImage)
-					.Image(FAutomationWindowStyle::Get().GetBrush("Icons.Delete"))
-				]
-			]
 		];
 
 	FToolBarBuilder ToolbarBuilder( InCommandList, FMultiBoxCustomization::None );
+	ToolbarBuilder.SetLabelVisibility(bIsLabelVisibilityEnabled ? EVisibility::Visible : EVisibility::Collapsed);
 	TWeakPtr<SAutomationWindow> AutomationWindow = SharedThis(this);
-	Local::FillToolbar(ToolbarBuilder, RunTests, TestPresets, AutomationWindow);
+	Local::FillToolbar(ToolbarBuilder, RunTests, DeviceGroups, TestPresets, AutomationWindow);
 
 	// Create the tool bar!
 	return
@@ -908,96 +985,122 @@ TSharedRef< SWidget > SAutomationWindow::MakeAutomationWindowToolBar( const TSha
 }
 END_SLATE_FUNCTION_BUILD_OPTIMIZATION
 
-EVisibility SAutomationWindow::HandlePresetComboVisibility( ) const
+
+TSharedRef< SWidget > SAutomationWindow::MakeAutomationFilterToolBar(const TSharedRef<FUICommandList>& InCommandList, TSharedPtr<class SAutomationWindow> InAutomationWindow)
 {
-	return bAddingTestPreset ? EVisibility::Hidden : EVisibility::Visible;
+	return InAutomationWindow->MakeAutomationFilterToolBar(InCommandList);
 }
 
-EVisibility SAutomationWindow::HandlePresetTextVisibility( ) const
+BEGIN_SLATE_FUNCTION_BUILD_OPTIMIZATION
+TSharedRef< SWidget > SAutomationWindow::MakeAutomationFilterToolBar(const TSharedRef<FUICommandList>& InCommandList)
 {
-	return bAddingTestPreset ? EVisibility::Visible : EVisibility::Hidden;
+	struct Local
+	{
+		static void FillToolbar(FToolBarBuilder& ToolbarBuilder)
+		{
+			ToolbarBuilder.BeginSection("Filters");
+			{
+				ToolbarBuilder.AddWidget(SNew(SSpacer));
+				ToolbarBuilder.AddToolBarButton(FAutomationWindowCommands::Get().ErrorFilter);
+				ToolbarBuilder.AddToolBarButton(FAutomationWindowCommands::Get().WarningFilter);
+				ToolbarBuilder.AddToolBarButton(FAutomationWindowCommands::Get().ExcludedTestsFilter);
+				ToolbarBuilder.AddToolBarButton(FAutomationWindowCommands::Get().DeveloperDirectoryContent);
+				ToolbarBuilder.AddWidget(SNew(SSpacer));
+			}
+			ToolbarBuilder.EndSection();
+		}
+	};
+
+	FToolBarBuilder ToolbarBuilder(InCommandList, FMultiBoxCustomization::None);
+	ToolbarBuilder.SetLabelVisibility(EVisibility::Collapsed);
+	TWeakPtr<SAutomationWindow> AutomationWindow = SharedThis(this);
+	Local::FillToolbar(ToolbarBuilder);
+
+	// Create the tool bar!
+	return
+		SNew(SHorizontalBox)
+		+ SHorizontalBox::Slot()
+		.AutoWidth()
+		[
+			SNew(SBorder)
+			.Padding(0)
+			.BorderImage(FAutomationWindowStyle::Get().GetBrush("NoBorder"))
+			.IsEnabled(FSlateApplication::Get().GetNormalExecutionAttribute())
+			[
+				ToolbarBuilder.MakeWidget()
+			]
+		];
 }
+END_SLATE_FUNCTION_BUILD_OPTIMIZATION
 
-EVisibility SAutomationWindow::HandleGroupsVisibility() const
+EVisibility SAutomationWindow::HandlePresetComboVisibility( ) const
 {
-	return GroupComboList.Num() == 1 ? EVisibility::Collapsed : EVisibility::Visible;
+	return bAddingTestPreset ? EVisibility::Hidden : EVisibility::Visible;
 }
 
-bool SAutomationWindow::IsAddButtonEnabled() const
+EVisibility SAutomationWindow::HandlePresetTextVisibility( ) const
 {
-	return !bAddingTestPreset && IsAutomationControllerIdle();
+	return bAddingTestPreset || bRenamingTestPreset ? EVisibility::Visible : EVisibility::Hidden;
 }
 
-bool SAutomationWindow::IsSaveButtonEnabled() const
+EVisibility SAutomationWindow::HandleGroupsVisibility() const
 {
-	return (!bAddingTestPreset && SelectedPreset.IsValid() && IsAutomationControllerIdle());
+	return GroupComboList.Num() == 1 ? EVisibility::Collapsed : EVisibility::Visible;
 }
 
-bool SAutomationWindow::IsRemoveButtonEnabled() const
+bool SAutomationWindow::IsSaveButtonEnabled() const
 {
-	return (!bAddingTestPreset && SelectedPreset.IsValid() && IsAutomationControllerIdle());
+	return (!bAddingTestPreset && !bRenamingTestPreset && SelectedPreset.IsValid() && IsAutomationControllerIdle());
 }
 
 void SAutomationWindow::HandlePresetTextCommited( const FText& CommittedText, ETextCommit::Type CommitType )
 {
 	if( CommitType == ETextCommit::OnEnter )
 	{
-		bAddingTestPreset = false;
-		if ( CommittedText.IsEmpty() )
+		if ( bAddingTestPreset )
 		{
-			return;
-		}
+			bAddingTestPreset = false;
+			if (CommittedText.IsEmpty())
+			{
+				return;
+			}
 
-		TArray<FString> EnabledTests;
-		AutomationController->GetEnabledTestNames(EnabledTests);
-		AutomationPresetPtr NewPreset = TestPresetManager->AddNewPreset(CommittedText, EnabledTests);
-		if ( NewPreset.IsValid() )
+			TArray<FString> EnabledTests;
+			AutomationController->GetEnabledTestNames(EnabledTests);
+			AutomationPresetPtr NewPreset = TestPresetManager->AddNewPreset(CommittedText, EnabledTests);
+			if (NewPreset.IsValid())
+			{
+				SelectedPreset = NewPreset;
+
+				PresetTextBox->SetText(FText());
+			}
+		}
+		else if (bRenamingTestPreset)
 		{
-			PresetComboBox->SetSelectedItem(NewPreset);
-			SelectedPreset = NewPreset;
+			bRenamingTestPreset = false;
+			if (CommittedText.IsEmpty())
+			{
+				return;
+			}
 
+			SelectedPreset->SetName(CommittedText);
 			PresetTextBox->SetText(FText());
 		}
+
+		PresetComboButtonText->SetText(CommittedText);
 	}
 	else if( CommitType == ETextCommit::OnCleared || CommitType == ETextCommit::OnUserMovedFocus )
 	{
-		if( bAddingTestPreset )
+		if( bAddingTestPreset || bRenamingTestPreset )
 		{
 			bAddingTestPreset = false;
+			bRenamingTestPreset = false;
 			SelectedPreset = nullptr;
-			PresetComboBox->ClearSelection();
 			PresetTextBox->SetText(FText());
 		}
 	}
 }
 
-void SAutomationWindow::HandlePresetChanged( TSharedPtr<FAutomationTestPreset> Item, ESelectInfo::Type SelectInfo )
-{
-	if( Item.IsValid() )
-	{
-		SelectedPreset = Item;
-		AutomationController->SetEnabledTests(Item->GetEnabledTests());
-		TestTable->RequestTreeRefresh();
-		
-		//Expand selected items
-		TestTable->ClearExpandedItems();
-		TArray< TSharedPtr< IAutomationReport > >& TestReports = AutomationController->GetReports();
-		for( int32 Index = 0; Index < TestReports.Num(); Index++ )
-		{
-			ExpandEnabledTests(TestReports[Index]);
-		}
-	}
-	else
-	{
-		SelectedPreset.Reset();
-
-		TArray<FString> EnabledTests;
-		AutomationController->SetEnabledTests(EnabledTests);
-		TestTable->ClearExpandedItems();
-		TestTable->RequestTreeRefresh();
-	}
-}
-
 void SAutomationWindow::HandleRequesteFilterChanged(TSharedPtr<FString> Item, ESelectInfo::Type SelectInfo)
 {
 	const int32 EntryIndex = RequestedFilterComboList.Find(Item);
@@ -1053,12 +1156,6 @@ void SAutomationWindow::ExpandEnabledTests( TSharedPtr< IAutomationReport > InRe
 	}
 }
 
-FReply SAutomationWindow::HandleNewPresetClicked()
-{
-	bAddingTestPreset = true;
-	return FReply::Handled().SetUserFocus(PresetTextBox.ToSharedRef(), EFocusCause::SetDirectly);
-}
-
 FReply SAutomationWindow::HandleSavePresetClicked()
 {
 	if(SelectedPreset.IsValid())
@@ -1071,15 +1168,30 @@ FReply SAutomationWindow::HandleSavePresetClicked()
 	return FReply::Handled();
 }
 
-FReply SAutomationWindow::HandleRemovePresetClicked()
+void SAutomationWindow::OnNewPresetClicked()
 {
-	if(SelectedPreset.IsValid())
+	bAddingTestPreset = true;
+	FSlateApplication::Get().SetUserFocus(0, PresetTextBox.ToSharedRef(), EFocusCause::SetDirectly);
+}
+
+void SAutomationWindow::OnRemovePresetClicked()
+{
+	if ( SelectedPreset.IsValid() )
 	{
 		TestPresetManager->RemovePreset(SelectedPreset.ToSharedRef());
 		SelectedPreset = nullptr;
-		PresetComboBox->ClearSelection();
+		PresetComboButtonText->SetText(LOCTEXT("AutomationPresetComboLabel", "None"));
+	}
+}
+
+void SAutomationWindow::OnRenamePresetClicked()
+{
+	if ( SelectedPreset.IsValid() )
+	{
+		bRenamingTestPreset = true;
+		PresetTextBox->SetText(SelectedPreset->GetName());
+		FSlateApplication::Get().SetUserFocus(0, PresetTextBox.ToSharedRef(), EFocusCause::SetDirectly);
 	}
-	return FReply::Handled();
 }
 
 FText SAutomationWindow::GetPresetComboText() const
@@ -1118,20 +1230,6 @@ FText SAutomationWindow::GetGroupComboText() const
 	}
 }
 
-TSharedRef<SWidget> SAutomationWindow::GeneratePresetComboItem(TSharedPtr<FAutomationTestPreset> InItem)
-{
-	if ( InItem.IsValid() )
-	{
-		return SNew(STextBlock)
-			.Text(InItem->GetName());
-	}
-	else
-	{
-		return SNew(STextBlock)
-			.Text(LOCTEXT("AutomationPreset_None", "None"));
-	}
-}
-
 TSharedRef<SWidget> SAutomationWindow::GenerateRequestedFilterComboItem(TSharedPtr<FString> InItem)
 {
 	return SNew(STextBlock)
@@ -1192,6 +1290,58 @@ TSharedRef< SWidget > SAutomationWindow::GenerateGroupOptionsMenuContent( )
 	return MenuBuilder.MakeWidget();
 }
 
+TSharedRef< SWidget > SAutomationWindow::GeneratePresetsMenuContent( TWeakPtr<class SAutomationWindow> InAutomationWindow )
+{
+	TSharedPtr<SAutomationWindow> AutomationWindow(InAutomationWindow.Pin());
+	if (AutomationWindow.IsValid())
+	{
+		return AutomationWindow->GenerateGroupOptionsMenuContent();
+	}
+
+	//Return empty menu
+	FMenuBuilder MenuBuilder(true, nullptr);
+	return MenuBuilder.MakeWidget();
+}
+
+TSharedRef< SWidget > SAutomationWindow::GeneratePresetsMenuContent()
+{
+	const bool bShouldCloseWindowAfterMenuSelection = true;
+	FMenuBuilder MenuBuilder(bShouldCloseWindowAfterMenuSelection, AutomationWindowActions);
+	MenuBuilder.BeginSection("AutomationWindowNewPresetGroup", LOCTEXT("GroupNewPreset", "New"));
+	MenuBuilder.AddMenuEntry(FAutomationWindowCommands::Get().CreateNewPreset);
+	MenuBuilder.EndSection();
+	const TArray<AutomationPresetPtr>& Presets = TestPresetManager->GetAllPresets();
+	TSharedPtr<TArray<TSharedPtr<SCheckBox>>> CheckBoxesSP = MakeShared<TArray<TSharedPtr<SCheckBox>>>();
+
+	MenuBuilder.BeginSection("AutomationWindowCurrentPresetGroup", LOCTEXT("GroupCurrentPreset", "Current Preset"));
+	{
+		MenuBuilder.AddMenuEntry(FAutomationWindowCommands::Get().RenameCurrentPreset);
+		MenuBuilder.AddMenuEntry(FAutomationWindowCommands::Get().RemoveCurrentPreset);
+
+		
+		for (int32 i = 0; i < Presets.Num(); i++)
+		{
+			const auto& preset = Presets[i];
+			TSharedPtr<SCheckBox> CheckBox =
+				SNew(SCheckBox)
+				.Style(FAutomationWindowStyle::Get(), "AutomationWindow.ToggleButton")
+				.IsChecked(preset.IsValid() && SelectedPreset.IsValid() && preset->GetID() == SelectedPreset->GetID())
+				.OnCheckStateChanged(this, &SAutomationWindow::HandlePresetCheckStateChanged, i, CheckBoxesSP)
+				.Padding(FMargin(4.0f, 0.0f))
+				.IsEnabled(this, &SAutomationWindow::IsAutomationControllerIdle)
+				.Content()
+				[
+					SNew(STextBlock)
+					.Text(preset.IsValid() ? preset.ToSharedRef()->GetName() : LOCTEXT("AutomationPreset_None", "None"))
+				];
+			MenuBuilder.AddWidget(CheckBox.ToSharedRef(), FText::GetEmpty());
+			CheckBoxesSP->Add(CheckBox);
+		}
+	}
+
+	return MenuBuilder.MakeWidget();
+}
+
 /** Returns if full size screen shots are enabled */
 ECheckBoxState SAutomationWindow::IsDeviceGroupCheckBoxIsChecked(const int32 DeviceGroupFlag) const
 {
@@ -1211,6 +1361,40 @@ void SAutomationWindow::HandleDeviceGroupCheckStateChanged(ECheckBoxState CheckB
 	TestTable->RebuildList();
 }
 
+void SAutomationWindow::HandlePresetCheckStateChanged(ECheckBoxState CheckBoxState, const int32 EntryIndex, TSharedPtr<TArray<TSharedPtr<SCheckBox>>> CheckBoxes)
+{
+	if (CheckBoxState != ECheckBoxState::Checked)
+	{
+		SelectedPreset = nullptr;
+		PresetComboButtonText->SetText(LOCTEXT("AutomationPreset_None", "None"));
+		return;
+	}
+	for (int32 i = 0; i < CheckBoxes->Num(); i++)
+	{
+		if (EntryIndex != i)
+		{
+			(*CheckBoxes)[i].Get()->SetIsChecked(false);
+		}
+	}
+
+	SelectedPreset = *(&TestPresetManager->GetAllPresets()[EntryIndex]);
+	PresetComboButtonText->SetText(SelectedPreset.IsValid() ? SelectedPreset.Get()->GetName() : LOCTEXT("AutomationPreset_None", "None"));
+
+	if (SelectedPreset.IsValid())
+	{
+		AutomationController->SetEnabledTests(SelectedPreset->GetEnabledTests());
+		TestTable->RequestTreeRefresh();
+
+		//Expand selected items
+		TestTable->ClearExpandedItems();
+		TArray< TSharedPtr< IAutomationReport > >& TestReports = AutomationController->GetReports();
+		for (int32 Index = 0; Index < TestReports.Num(); Index++)
+		{
+			ExpandEnabledTests(TestReports[Index]);
+		}
+	}
+}
+
 TSharedRef< SWidget > SAutomationWindow::GenerateTestsOptionsMenuContent( TWeakPtr<class SAutomationWindow> InAutomationWindow )
 {
 	TSharedPtr<SAutomationWindow> AutomationWindow(InAutomationWindow.Pin());
@@ -1652,44 +1836,16 @@ void SAutomationWindow::RebuildPlatformIcons()
 	int32 NumClusters = AutomationController->GetNumDeviceClusters();
 	for (int32 ClusterIndex = 0; ClusterIndex < NumClusters; ++ClusterIndex)
 	{
-		//find the right platform icon
-		FString DeviceTypeName = AutomationController->GetDeviceTypeName(ClusterIndex);
-		const PlatformInfo::FTargetPlatformInfo* TargetPlatformInfo = PlatformInfo::FindPlatformInfo(FName(*DeviceTypeName));
-		ensureMsgf(TargetPlatformInfo, TEXT("PlatformInfo::FindPlatformInfo() failed looking up %s"), *DeviceTypeName);
-		ensureMsgf(!TargetPlatformInfo || TargetPlatformInfo->DataDrivenPlatformInfo, TEXT("PlatformInfo::FTargetPlatformInfo had null DataDrivenPlatformInfo"));
-		const FDataDrivenPlatformInfo* DDPI = TargetPlatformInfo ? TargetPlatformInfo->DataDrivenPlatformInfo : nullptr;
-		FName DeviceImageName = DDPI ? DDPI->GetIconStyleName(EPlatformIconSize::Normal) : NAME_None;
-
-		const FSlateBrush* ImageToUse = FAutomationWindowStyle::Get().GetBrush(DeviceImageName);
-
 		PlatformsHBox->AddSlot()
 		.AutoWidth()
 		.MaxWidth(ColumnWidth)
 		[
-			SNew( SOverlay )
-			+ SOverlay::Slot()
-			[
-				SNew(SBorder)
-				.BorderImage(FAutomationWindowStyle::Get().GetBrush("ErrorReporting.Box") )
-				.HAlign(HAlign_Center)
-				.VAlign(VAlign_Center)
-				.Padding( FMargin(3,0) )
-				.BorderBackgroundColor( FSlateColor( FLinearColor( 1.0f, 0.0f, 1.0f, 0.0f ) ) )
-				.ToolTipText( CreateDeviceTooltip( ClusterIndex ) )
-				[
-					SNew(SImage)
-					.DesiredSizeOverride(FVector2D(16.f, 16.f))
-					.Image(ImageToUse)
-				]
-			]
-			+ SOverlay::Slot()
-			.HAlign(HAlign_Right)
-			.VAlign(VAlign_Bottom)
-			[
-				//Overlay how many devices are in the cluster
-				SNew( STextBlock )
-				.Text( this, &SAutomationWindow::OnGetNumDevicesInClusterString, ClusterIndex )
-			]
+			SNew(STextBlock)
+			.Text(FText::Format(
+				LOCTEXT("GameInstances_Header", "Instances ({0})"),
+				FText::AsNumber( AutomationController->GetNumDevicesInCluster(ClusterIndex) )
+			))
+			.ToolTipText(CreateDeviceTooltip(ClusterIndex))
 		];
 	}
 }
diff --git a/Engine/Source/Developer/AutomationWindow/Private/SAutomationWindow.h b/Engine/Source/Developer/AutomationWindow/Private/SAutomationWindow.h
--- a/Engine/Source/Developer/AutomationWindow/Private/SAutomationWindow.h
+++ b/Engine/Source/Developer/AutomationWindow/Private/SAutomationWindow.h
@@ -35,12 +35,14 @@ template< typename ItemType > class TTextFilter;
 /** Columns for the test tree view */
 namespace AutomationTestWindowConstants
 {
+	const FName Checked( TEXT("Checked") );
+	const FName Skipped( TEXT("Skipped") );
+	const FName SkippedOptions( TEXT("SkippedOptions") );
 	const FName Title( TEXT("Name") );
 	const FName SmokeTest( TEXT("SmokeTest") );
 	const FName RequiredDeviceCount( TEXT("RequiredDeviceCount") );
 	const FName Status( TEXT("Status") );
 	const FName Timing( TEXT("Timing") );
-	const FName IsToBeSkipped(TEXT("ToBeSkipped"));
 }
 
 
@@ -186,6 +188,14 @@ private:
 	static TSharedRef< SWidget > MakeAutomationWindowToolBar( const TSharedRef<FUICommandList>& InCommandList, TSharedPtr<class SAutomationWindow> InLevelEditor );
 	TSharedRef< SWidget > MakeAutomationWindowToolBar( const TSharedRef<FUICommandList>& InCommandList );
 
+	/**
+	 * Static: Creates a filter toolbar widget for the automation window.
+	 *
+	 * @return The new widget.
+	 */
+	static TSharedRef< SWidget > MakeAutomationFilterToolBar(const TSharedRef<FUICommandList>& InCommandList, TSharedPtr<class SAutomationWindow> InLevelEditor);
+	TSharedRef< SWidget > MakeAutomationFilterToolBar(const TSharedRef<FUICommandList>& InCommandList);
+
 	/**
 	 * Static: Creates the test options menu widget.
 	 *
@@ -202,6 +212,14 @@ private:
 	static TSharedRef< SWidget >GenerateGroupOptionsMenuContent( TWeakPtr<class SAutomationWindow> InAutomationWindow );
 	TSharedRef< SWidget > GenerateGroupOptionsMenuContent( );
 
+	/**
+	 * Static: Creates the Presets menu widget.
+	 *
+	 * @return The new widget.
+	 */
+	static TSharedRef< SWidget >GeneratePresetsMenuContent( TWeakPtr<class SAutomationWindow> InAutomationWindow );
+	TSharedRef< SWidget > GeneratePresetsMenuContent();
+
 	/**
 	* Static: Creates the test history options menu widget.
 	*
@@ -210,13 +228,6 @@ private:
 	static TSharedRef< SWidget > GenerateTestHistoryMenuContent(TWeakPtr<class SAutomationWindow> InAutomationWindow);
 	TSharedRef< SWidget > GenerateTestHistoryMenuContent();
 
-	/**
-	 * Creates a combo item for the preset list.
-	 *
-	 * @return New combo item widget.
-	 */
-	TSharedRef<SWidget> GeneratePresetComboItem(TSharedPtr<FAutomationTestPreset> InItem);
-
 	/**
 	 * Creates a combo item for the requested filter.
 	 *
@@ -338,7 +349,10 @@ private:
 	
 	/** Toggles a device group flag */
 	void HandleDeviceGroupCheckStateChanged(ECheckBoxState CheckBoxState, const int32 DeviceGroupFlag);
-	
+
+	/** Sitches presets in the Presets Menu of Preset ComboButton */
+	void HandlePresetCheckStateChanged(ECheckBoxState CheckBoxState, const int32 EntryIndex, TSharedPtr<TArray<TSharedPtr<SCheckBox>>> CheckBoxes);
+
 	/** Sets the number of times to repeat the tests */
 	void OnChangeRepeatCount(int32 InNewValue);
 	
@@ -454,24 +468,21 @@ private:
 	bool ExpandToTest(TSharedPtr<IAutomationReport> InRoot, TSharedPtr<IAutomationReport> InReport);
 #endif
 
-	/** Handles the new preset button being clicked. */
-	FReply HandleNewPresetClicked();
-
 	/** Handles the save preset button being clicked. */
 	FReply HandleSavePresetClicked();
 
-	/** Handles the remove preset button being clicked. */
-	FReply HandleRemovePresetClicked();
+	/** Handles the create new preset menu item being clicked. */
+	void OnNewPresetClicked();
+
+	/** Handles the remove preset menu item being clicked. */
+	void OnRemovePresetClicked();
 
-	/** Should the add preset button be enabled. */
-	bool IsAddButtonEnabled() const;
+	/** Handles the rename preset menu item being clicked. */
+	void OnRenamePresetClicked();
 
 	/** Should the save preset button be enabled. */
 	bool IsSaveButtonEnabled() const;
 
-	/** Should the remove preset button be enabled. */
-	bool IsRemoveButtonEnabled() const;
-
 	/** Handles if the preset combo box should be visible. */
 	EVisibility HandlePresetComboVisibility( ) const;
 
@@ -484,8 +495,6 @@ private:
 	/** Called when the user commits the text in the add preset text box. */
 	void HandlePresetTextCommited( const FText& CommittedText, ETextCommit::Type CommitType );
 
-	/** Called when the user selects a new preset from the preset combo box. */
-	void HandlePresetChanged( TSharedPtr<FAutomationTestPreset> Item, ESelectInfo::Type SelectInfo );
 	/** Called when the user changes the requested test filter */
 	void HandleRequesteFilterChanged(TSharedPtr<FString> Item, ESelectInfo::Type SelectInfo);
 	/** Called when the user changes the test group */
@@ -620,17 +629,23 @@ private:
 	/** Which type of window style to use for the test background. */
 	EAutomationTestBackgroundStyle::Type TestBackgroundType;
 
-	/** True if we are creating a new preset (The add preset text box is visible). */
+	/** True if we are creating a new preset (The add/rename preset text box is visible). */
 	bool bAddingTestPreset;
 
+	/** True if we are renaming a new preset (The add/rename preset text box is visible). */
+	bool bRenamingTestPreset;
+
+	/** Flag to enable display of the text labels in the toolbar. */
+	bool bIsLabelVisibilityEnabled;
+
 	/** Holds a pointer to the preset manager. */
 	TSharedPtr<FAutomationTestPresetManager> TestPresetManager;
 
 	/** Holds the currently selected preset. */
 	TSharedPtr<FAutomationTestPreset> SelectedPreset;
 
-	/** Holds a pointer to the preset combo box widget. */
-	TSharedPtr< SComboBox< TSharedPtr<FAutomationTestPreset> > > PresetComboBox;
+	/** Holds a pointer to the preset combo button widget. */
+	TSharedPtr< STextBlock > PresetComboButtonText;
 
 	/** Holds a pointer to the preset combo box widget. */
 	TSharedPtr< SComboBox< TSharedPtr<FString> > >	RequestedFilterComboBox;
diff --git a/Engine/Source/Developer/SessionFrontend/Private/Widgets/SSessionFrontend.cpp b/Engine/Source/Developer/SessionFrontend/Private/Widgets/SSessionFrontend.cpp
--- a/Engine/Source/Developer/SessionFrontend/Private/Widgets/SSessionFrontend.cpp
+++ b/Engine/Source/Developer/SessionFrontend/Private/Widgets/SSessionFrontend.cpp
@@ -55,7 +55,7 @@ void SSessionFrontend::Construct( const FArguments& InArgs, const TSharedRef<SDo
 
 	TabManager->RegisterTabSpawner(AutomationTabId, FOnSpawnTab::CreateRaw(this, &SSessionFrontend::HandleTabManagerSpawnTab, AutomationTabId))
 		.SetDisplayName(LOCTEXT("AutomationTabTitle", "Automation"))
-		.SetIcon(FSlateIcon(FAppStyle::GetAppStyleSetName(), "SessionFrontEnd.Tabs.Tools"))
+		.SetIcon(FSlateIcon(FAppStyle::GetAppStyleSetName(), "SessionFrontEnd.Tabs.Automation"))
 		.SetGroup(AppMenuGroup);
 
 	TabManager->RegisterTabSpawner(SessionBrowserTabId, FOnSpawnTab::CreateRaw(this, &SSessionFrontend::HandleTabManagerSpawnTab, SessionBrowserTabId))
@@ -65,12 +65,12 @@ void SSessionFrontend::Construct( const FArguments& InArgs, const TSharedRef<SDo
 
 	TabManager->RegisterTabSpawner(SessionConsoleTabId, FOnSpawnTab::CreateRaw(this, &SSessionFrontend::HandleTabManagerSpawnTab, SessionConsoleTabId))
 		.SetDisplayName(LOCTEXT("ConsoleTabTitle", "Console"))
-		.SetIcon(FSlateIcon(FAppStyle::GetAppStyleSetName(), "SessionFrontEnd.Tabs.Tools"))
+		.SetIcon(FSlateIcon(FAppStyle::GetAppStyleSetName(), "SessionFrontEnd.Tabs.Console"))
 		.SetGroup(AppMenuGroup);
 
 	TabManager->RegisterTabSpawner(SessionScreenTabId, FOnSpawnTab::CreateRaw(this, &SSessionFrontend::HandleTabManagerSpawnTab, SessionScreenTabId))
 		.SetDisplayName(LOCTEXT("ScreenTabTitle", "Screen Comparison"))
-		.SetIcon(FSlateIcon(FAppStyle::GetAppStyleSetName(), "SessionFrontEnd.Tabs.Tools"))
+		.SetIcon(FSlateIcon(FAppStyle::GetAppStyleSetName(), "SessionFrontEnd.Tabs.ScreenComparison"))
 		.SetGroup(AppMenuGroup);
 
 	TabManager->RegisterTabSpawner(ProfilerTabId, FOnSpawnTab::CreateRaw(this, &SSessionFrontend::HandleTabManagerSpawnTab, ProfilerTabId))
diff --git a/Engine/Source/Developer/TraceInsights/Private/Insights/InsightsStyle.cpp b/Engine/Source/Developer/TraceInsights/Private/Insights/InsightsStyle.cpp
--- a/Engine/Source/Developer/TraceInsights/Private/Insights/InsightsStyle.cpp
+++ b/Engine/Source/Developer/TraceInsights/Private/Insights/InsightsStyle.cpp
@@ -302,7 +302,7 @@ void FInsightsStyle::FStyle::Initialize()
 	Set("Icons.RemoveGraphSeries", new CORE_IMAGE_BRUSH_SVG("Starship/Common/close", Icon16x16));
 
 	Set("Icons.TestAutomation", new EDITOR_IMAGE_BRUSH_SVG("Starship/Common/TestAutomation", Icon16x16));
-	Set("Icons.Test", new EDITOR_IMAGE_BRUSH_SVG("Starship/Common/Test", Icon16x16));
+	Set("Icons.Test", new CORE_IMAGE_BRUSH_SVG("Starship/Common/Test", Icon16x16));
 
 	Set("Icons.Debug", new EDITOR_IMAGE_BRUSH_SVG("Starship/Common/bug", Icon16x16));
 	Set("Icons.Debug.ToolBar", new EDITOR_IMAGE_BRUSH_SVG("Starship/Common/bug", Icon20x20));
diff --git a/Engine/Source/Editor/EditorStyle/Private/StarshipStyle.cpp b/Engine/Source/Editor/EditorStyle/Private/StarshipStyle.cpp
--- a/Engine/Source/Editor/EditorStyle/Private/StarshipStyle.cpp
+++ b/Engine/Source/Editor/EditorStyle/Private/StarshipStyle.cpp
@@ -992,7 +992,7 @@ void FStarshipEditorStyle::FStyle::SetupGeneralStyles()
 
 	// Automation Tools Menu
 	{
-		Set("AutomationTools.MenuIcon", new IMAGE_BRUSH_SVG("Starship/Common/AutomationTools", Icon16x16));
+		Set("AutomationTools.MenuIcon", new CORE_IMAGE_BRUSH_SVG("Starship/Common/AutomationTools", Icon16x16));
 		Set("AutomationTools.TestAutomation", new IMAGE_BRUSH_SVG("Starship/Common/TestAutomation", Icon16x16));
 	}
 
@@ -1018,8 +1018,11 @@ void FStarshipEditorStyle::FStyle::SetupGeneralStyles()
 
 	// Session Frontend Window
 	{
-		Set("SessionFrontEnd.TabIcon", new IMAGE_BRUSH_SVG( "Starship/Common/SessionFrontend", Icon16x16 ) );
+		Set("SessionFrontEnd.TabIcon", new CORE_IMAGE_BRUSH_SVG( "Starship/Common/SessionFrontend", Icon16x16 ) );
 		Set("SessionFrontEnd.Tabs.Tools", new CORE_IMAGE_BRUSH( "/Icons/icon_tab_Tools_16x", Icon16x16 ) );
+		Set("SessionFrontEnd.Tabs.Console", new CORE_IMAGE_BRUSH_SVG("Starship/Common/Console", Icon16x16));
+		Set("SessionFrontEnd.Tabs.Automation", new CORE_IMAGE_BRUSH_SVG("Starship/Common/AutomationTools", Icon16x16));
+		Set("SessionFrontEnd.Tabs.ScreenComparison", new CORE_IMAGE_BRUSH_SVG("Starship/Common/Monitor", Icon16x16));
 	}
 
 	// Undo History Window
@@ -3863,7 +3866,7 @@ void FStarshipEditorStyle::FStyle::SetupGraphEditorStyles()
 		Set( "GraphEditor.Sequence_16x", new IMAGE_BRUSH_SVG("Starship/GraphEditors/Sequence", Icon16x16));
 		Set( "GraphEditor.Branch_16x", new IMAGE_BRUSH_SVG("Starship/GraphEditors/Branch", Icon16x16));
 		Set( "GraphEditor.SpawnActor_16x", new IMAGE_BRUSH_SVG("Starship/GraphEditors/SpawnActor", Icon16x16));
-		Set( "GraphEditor.PadEvent_16x", new IMAGE_BRUSH_SVG("Starship/Common/PlayerController", Icon16x16));
+		Set( "GraphEditor.PadEvent_16x", new CORE_IMAGE_BRUSH_SVG("Starship/Common/PlayerController", Icon16x16));
 		Set( "GraphEditor.MouseEvent_16x", new IMAGE_BRUSH_SVG("Starship/GraphEditors/MouseEvent", Icon16x16));
 		Set( "GraphEditor.KeyEvent_16x", new IMAGE_BRUSH_SVG("Starship/Common/ViewportControls", Icon16x16));
 		Set( "GraphEditor.TouchEvent_16x", new IMAGE_BRUSH_SVG("Starship/Common/TouchInterface", Icon16x16));
@@ -4487,7 +4490,7 @@ void FStarshipEditorStyle::FStyle::SetupLevelEditorStyle()
 		Set("PlacementBrowser.Icons.BSP",           new IMAGE_BRUSH_SVG("Starship/Common/Geometry",       Icon16x16));
 		Set("PlacementBrowser.Icons.Volumes",       new IMAGE_BRUSH_SVG("Starship/Common/Volumes",        Icon16x16));
 		Set("PlacementBrowser.Icons.All",           new IMAGE_BRUSH_SVG("Starship/Common/AllClasses",     Icon16x16));
-		Set("PlacementBrowser.Icons.Testing",       new IMAGE_BRUSH_SVG("Starship/Common/Test",           Icon16x16));
+		Set("PlacementBrowser.Icons.Testing",       new CORE_IMAGE_BRUSH_SVG("Starship/Common/Test",           Icon16x16));
 
 		Set( "ContentPalette.ShowAllPlaceables", new IMAGE_BRUSH( "Icons/icon_Placement_FilterAll_20px", Icon20x20 ) );
 		Set( "ContentPalette.ShowAllPlaceables.Small", new IMAGE_BRUSH( "Icons/icon_Placement_FilterAll_20px", Icon20x20 ) );
@@ -4647,7 +4650,7 @@ void FStarshipEditorStyle::FStyle::SetupLevelEditorStyle()
 		Set("ShowFlagsMenu.SubMenu.LightingComponents", new IMAGE_BRUSH_SVG("Starship/Common/LightingComponents", Icon16x16));
 		Set("ShowFlagsMenu.SubMenu.LightingFeatures", new IMAGE_BRUSH_SVG("Starship/Common/LightingFeatures", Icon16x16));
 		Set("ShowFlagsMenu.SubMenu.Lumen", new IMAGE_BRUSH_SVG("Starship/Common/LightingFeatures", Icon16x16));
-		Set("ShowFlagsMenu.SubMenu.Developer", new IMAGE_BRUSH_SVG("Starship/Common/Developer", Icon16x16));
+		Set("ShowFlagsMenu.SubMenu.Developer", new CORE_IMAGE_BRUSH_SVG("Starship/Common/Developer", Icon16x16));
 		Set("ShowFlagsMenu.SubMenu.Visualize", new IMAGE_BRUSH_SVG("Starship/Common/Visualize", Icon16x16));
 		Set("ShowFlagsMenu.SubMenu.Advanced", new CORE_IMAGE_BRUSH_SVG("Starship/Common/Advanced", Icon16x16));
 
