Add missing header to SAutomationWindow.cpp

diff --git a/Engine/Source/Developer/AutomationWindow/Private/SAutomationWindow.cpp b/Engine/Source/Developer/AutomationWindow/Private/SAutomationWindow.cpp
--- a/Engine/Source/Developer/AutomationWindow/Private/SAutomationWindow.cpp
+++ b/Engine/Source/Developer/AutomationWindow/Private/SAutomationWindow.cpp
@@ -19,6 +19,7 @@
 #include "Framework/MultiBox/MultiBoxBuilder.h"
 #include "Widgets/Input/SEditableTextBox.h"
 #include "Widgets/Input/SButton.h"
+#include "Widgets/Layout/SSpacer.h"
 #include "Widgets/Layout/SScrollBox.h"
 #include "Widgets/Input/SCheckBox.h"
 #include "Widgets/Input/SSpinBox.h"
