Place test result bar in Test Automation window result below GameInstance

diff --git a/Engine/Source/Developer/AutomationWindow/Private/SAutomationGraphicalResultBox.cpp b/Engine/Source/Developer/AutomationWindow/Private/SAutomationGraphicalResultBox.cpp
--- a/Engine/Source/Developer/AutomationWindow/Private/SAutomationGraphicalResultBox.cpp
+++ b/Engine/Source/Developer/AutomationWindow/Private/SAutomationGraphicalResultBox.cpp
@@ -254,8 +254,9 @@ void SAutomationGraphicalResultBox::CreateWidgets()
 					];
 			}
 
-			
-			GridContainer->AddSlot(1,RowCounter)
+			RowCounter++;
+
+			GridContainer->AddSlot(0,RowCounter)
 				.HAlign(HAlign_Fill)
 				.VAlign(VAlign_Center)
 				.Padding( FMargin(1,3) )
