Prevent vanishing of test results bar

diff --git a/Engine/Source/Developer/AutomationWindow/Private/SAutomationGraphicalResultBox.cpp b/Engine/Source/Developer/AutomationWindow/Private/SAutomationGraphicalResultBox.cpp
--- a/Engine/Source/Developer/AutomationWindow/Private/SAutomationGraphicalResultBox.cpp
+++ b/Engine/Source/Developer/AutomationWindow/Private/SAutomationGraphicalResultBox.cpp
@@ -203,6 +203,11 @@ void SAutomationGraphicalResultBox::CreateWidgets()
 
 			TSharedRef<SHorizontalBox> TestContainer = SNew(SHorizontalBox);
 
+			// If there are too many tests then result blocks may become invisible because they get too tiny
+			// and margins are covering them up. Hence reducing the margins with relation to the number of tests.
+			int32 TestsNum = DeviceIt->Tests.Num();
+			float HorizontalMargin = TestsNum < 30 ? 1.0 : (TestsNum < 100 ? 0.5 : (TestsNum < 1000 ? 0.25 : 0.1));
+
 			for( int32 TestIndex = 0; TestIndex < DeviceIt->Tests.Num(); ++TestIndex)
 			{
 				FTestResults* TestIt = &DeviceIt->Tests[TestIndex];
@@ -214,7 +219,7 @@ void SAutomationGraphicalResultBox::CreateWidgets()
 				TestContainer->AddSlot()
 					.HAlign(HAlign_Fill)
 					.VAlign(VAlign_Fill)
-					.Padding( FMargin(1,5) )
+					.Padding( FMargin(HorizontalMargin, 5) )
 					.FillWidth(TestIt->Duration)
 					[
 						SNew(SOverlay)
@@ -259,8 +264,8 @@ void SAutomationGraphicalResultBox::CreateWidgets()
 			GridContainer->AddSlot(0,RowCounter)
 				.HAlign(HAlign_Fill)
 				.VAlign(VAlign_Center)
-				.Padding( FMargin(1,3) )
-				.ColumnSpan(9)
+				.Padding( FMargin(1,0) )
+				.ColumnSpan(11)
 				[
 					SNew(SHorizontalBox)
 					+SHorizontalBox::Slot()
