Fixed UE-178837.
Added #if WITH_EDITOR to SAutomationWindow::HandleSessionManagerCanSelectSession

diff --git a/Engine/Source/Developer/AutomationWindow/Private/SAutomationWindow.cpp b/Engine/Source/Developer/AutomationWindow/Private/SAutomationWindow.cpp
--- a/Engine/Source/Developer/AutomationWindow/Private/SAutomationWindow.cpp
+++ b/Engine/Source/Developer/AutomationWindow/Private/SAutomationWindow.cpp
@@ -2323,11 +2323,16 @@ EVisibility SAutomationWindow::HandleSelectSessionOverlayVisibility( ) const
 
 void SAutomationWindow::HandleSessionManagerCanSelectSession( const TSharedPtr<ISessionInfo>& Session, bool& CanSelect )
 {
+	// We are using a compilation condition here as FMessageDialog::Open will show dialog only if GIsEditor && !IsRunningCommandlet() && FCoreDelegates::ModalErrorMessage.IsBound()
+	// else it always gives us EAppReturnType::No answer in EAppMsgType::YesNo case.
+	// The result is the sessions cannot be chosen in not editor case (ie: UnrealFrontEnd)   
+#if WITH_EDITOR 
 	if (ActiveSession.IsValid() && AutomationController->CheckTestResultsAvailable())
 	{
 		EAppReturnType::Type Result = FMessageDialog::Open(EAppMsgType::YesNo, LOCTEXT("ChangeSessionDialog", "Are you sure you want to change sessions?\nAll automation results data will be lost"));
 		CanSelect = Result == EAppReturnType::Yes ? true : false;
 	}
+#endif //WITH_EDITOR
 }
 
 
