Tooltip in the Exclusion buttons in Session Frontend Automation Test no longer displays the reason

diff --git a/Engine/Source/Developer/AutomationWindow/Private/SAutomationTestItem.cpp b/Engine/Source/Developer/AutomationWindow/Private/SAutomationTestItem.cpp
--- a/Engine/Source/Developer/AutomationWindow/Private/SAutomationTestItem.cpp
+++ b/Engine/Source/Developer/AutomationWindow/Private/SAutomationTestItem.cpp
@@ -70,7 +70,7 @@ TSharedRef<SWidget> SAutomationTestItem::GenerateWidgetForColumn( const FName& C
 			[
 				SNew(SButton)
 				.ButtonStyle(FAutomationWindowStyle::Get(), "NoBorder")
-				.ToolTipText(LOCTEXT("ToggleExcluded", "Exclude"))
+				.ToolTipText(this, &SAutomationTestItem::GetExcludeReason)
 				.OnClicked(FOnClicked::CreateSP(this, &SAutomationTestItem::SetSkipFlag))
 				[
 					SNew(SImage)
