Session FrontEnd - Missing export menu and copy button when test results are displayed

diff --git a/Engine/Source/Developer/AutomationWindow/Private/SAutomationWindow.cpp b/Engine/Source/Developer/AutomationWindow/Private/SAutomationWindow.cpp
--- a/Engine/Source/Developer/AutomationWindow/Private/SAutomationWindow.cpp
+++ b/Engine/Source/Developer/AutomationWindow/Private/SAutomationWindow.cpp
@@ -495,11 +495,14 @@ void SAutomationWindow::Construct( const FArguments& InArgs, const IAutomationCo
 				]
 
 				+ SSplitter::Slot()
-					.Value(0.33f)
-					[
-						SNew(SOverlay)
+				.Value(0.33f)
+				[
 
-						+ SOverlay::Slot()
+						SNew(SVerticalBox)
+						+ SVerticalBox::Slot()
+						[
+							SNew(SOverlay)
+							+ SOverlay::Slot()
 							[
 								SNew(SBox)
 									.Visibility(this, &SAutomationWindow::GetTestGraphVisibility)
@@ -572,58 +575,58 @@ void SAutomationWindow::Construct( const FArguments& InArgs, const IAutomationCo
 									]
 							]
 
-						+ SOverlay::Slot()
-						[
-							SNew(SBox)
-							.Visibility(this, &SAutomationWindow::GetTestLogVisibility)
+							+ SOverlay::Slot()
 							[
-								//results panel
-								SNew( SVerticalBox )
-								+ SVerticalBox::Slot()
-								.AutoHeight()
-								[
-									SNew( STextBlock )
-									.Text( LOCTEXT("AutomationTest_Results", "Automation Test Results:") )
-								]
-
-								+ SVerticalBox::Slot()
-								.FillHeight(1.0f)
-								.Padding(0.0f, 4.0f, 0.0f, 0.0f)
+								SNew(SBox)
+								.Visibility(this, &SAutomationWindow::GetTestLogVisibility)
 								[
-									//list of results for the selected test
-									SNew(SBorder)
-									.BorderImage(FAutomationWindowStyle::Get().GetBrush("Brushes.Panel"))
+									//results panel
+									SNew( SVerticalBox )
+									+ SVerticalBox::Slot()
+									.AutoHeight()
 									[
-										SNew(SScrollBox)
-										.Orientation(EOrientation::Orient_Horizontal)
-										+SScrollBox::Slot()
-										[
-											SAssignNew(LogListView, SListView<TSharedPtr<FAutomationOutputMessage> >)
-											.ItemHeight(18)
-											.ListItemsSource(&LogMessages)
-											.SelectionMode(ESelectionMode::Multi)
-											.OnGenerateRow(this, &SAutomationWindow::OnGenerateWidgetForLog)
-											.OnSelectionChanged(this, &SAutomationWindow::HandleLogListSelectionChanged)
-										]
+										SNew( STextBlock )
+										.Text( LOCTEXT("AutomationTest_Results", "Automation Test Results:") )
 									]
-								]
 
-								+ SVerticalBox::Slot()
-								.AutoHeight()
-								.Padding(0.0f, 4.0f, 0.0f, 0.0f)
-								[
-									SNew(SBorder)
-									.BorderImage(FAutomationWindowStyle::Get().GetBrush("ToolPanel.GroupBorder"))
-									.Padding(FMargin(8.0f, 6.0f))
+									+ SVerticalBox::Slot()
+									.FillHeight(1.0f)
+									.Padding(0.0f, 4.0f, 0.0f, 0.0f)
 									[
-										// Add the command bar
-										SAssignNew(CommandBar, SAutomationWindowCommandBar, NotificationList)
-										.OnCopyLogClicked(this, &SAutomationWindow::HandleCommandBarCopyLogClicked)
+										//list of results for the selected test
+										SNew(SBorder)
+										.BorderImage(FAutomationWindowStyle::Get().GetBrush("Brushes.Panel"))
+										[
+											SNew(SScrollBox)
+											.Orientation(EOrientation::Orient_Horizontal)
+											+SScrollBox::Slot()
+											[
+												SAssignNew(LogListView, SListView<TSharedPtr<FAutomationOutputMessage> >)
+												.ItemHeight(18)
+												.ListItemsSource(&LogMessages)
+												.SelectionMode(ESelectionMode::Multi)
+												.OnGenerateRow(this, &SAutomationWindow::OnGenerateWidgetForLog)
+												.OnSelectionChanged(this, &SAutomationWindow::HandleLogListSelectionChanged)
+											]
+										]
 									]
 								]
 							]
 						]
-					]
+						+SVerticalBox::Slot()
+						.AutoHeight()
+						.Padding(0.0f, 4.0f, 0.0f, 0.0f)
+						[
+							SNew(SBorder)
+							.BorderImage(FAutomationWindowStyle::Get().GetBrush("ToolPanel.GroupBorder"))
+							.Padding(FMargin(8.0f, 6.0f))
+							[
+								// Add the command bar
+								SAssignNew(CommandBar, SAutomationWindowCommandBar, NotificationList)
+								.OnCopyLogClicked(this, &SAutomationWindow::HandleCommandBarCopyLogClicked)
+							]
+						]
+				]
 			]
 
 		+ SOverlay::Slot()
@@ -1742,6 +1745,8 @@ void SAutomationWindow::OnTestSelectionChanged(TSharedPtr<IAutomationReport> Sel
 			bHasChildTestSelected = true;
 		}
 	}
+
+	CommandBar->SetCopyButtonVisibility(GetTestLogVisibility());
 }
 
 void SAutomationWindow::UpdateTestLog(TSharedPtr<IAutomationReport> Selection)
diff --git a/Engine/Source/Developer/AutomationWindow/Private/SAutomationWindowCommandBar.cpp b/Engine/Source/Developer/AutomationWindow/Private/SAutomationWindowCommandBar.cpp
--- a/Engine/Source/Developer/AutomationWindow/Private/SAutomationWindowCommandBar.cpp
+++ b/Engine/Source/Developer/AutomationWindow/Private/SAutomationWindowCommandBar.cpp
@@ -27,6 +27,7 @@ void SAutomationWindowCommandBar::Construct( const FArguments& InArgs, const TSh
 			.HAlign(HAlign_Right)
 			[
 				SNew(SAutomationExportMenu, InNotificationList)
+				.Visibility(this, &SAutomationWindowCommandBar::GetExportMenuVisibility)
 			]
 
 		+ SHorizontalBox::Slot()
@@ -35,6 +36,7 @@ void SAutomationWindowCommandBar::Construct( const FArguments& InArgs, const TSh
 			[
 				// copy button
 				SAssignNew(CopyButton, SButton)
+					.Visibility(this, &SAutomationWindowCommandBar::GetCopyButtonVisibility)
 					.ContentPadding(FMargin(6.0, 2.0))
 					.IsEnabled(false)
 					.Text(LOCTEXT("AutomationCopyButtonText", "Copy"))
@@ -45,6 +47,23 @@ void SAutomationWindowCommandBar::Construct( const FArguments& InArgs, const TSh
 }
 
 
+void SAutomationWindowCommandBar::SetCopyButtonVisibility(EVisibility InVisibility)
+{
+	CopyButtonVisibility = (InVisibility == EVisibility::Visible) ? EVisibility::Visible : EVisibility::Collapsed;
+}
+
+
+EVisibility SAutomationWindowCommandBar::GetCopyButtonVisibility() const
+{
+	return CopyButtonVisibility;
+}
+
+
+EVisibility SAutomationWindowCommandBar::GetExportMenuVisibility() const
+{
+	return (CopyButtonVisibility == EVisibility::Visible) ? EVisibility::Hidden : EVisibility::Visible;
+}
+
 FReply SAutomationWindowCommandBar::HandleCopyButtonClicked( )
 {
 	if (OnCopyLogClicked.IsBound())
diff --git a/Engine/Source/Developer/AutomationWindow/Private/SAutomationWindowCommandBar.h b/Engine/Source/Developer/AutomationWindow/Private/SAutomationWindowCommandBar.h
--- a/Engine/Source/Developer/AutomationWindow/Private/SAutomationWindowCommandBar.h
+++ b/Engine/Source/Developer/AutomationWindow/Private/SAutomationWindowCommandBar.h
@@ -42,11 +42,17 @@ public:
 	 */
 	void SetNumLogMessages( int Count );
 
+	/** Sets visibility status for Copy button */
+	void SetCopyButtonVisibility(EVisibility InVisibility);
+
 private:
 
 	/** Handles clicking the copy log button. */
 	FReply HandleCopyButtonClicked( );
 
+	EVisibility GetCopyButtonVisibility() const;
+	EVisibility GetExportMenuVisibility() const;
+
 private:
 
  	/** Holds the copy log button. */
@@ -54,4 +60,7 @@ private:
 
  	/** Holds a delegate that is executed when the copy log button is clicked. */
  	FOnClicked OnCopyLogClicked;
+
+	/** Holds visibility status for Copy button */
+	EVisibility CopyButtonVisibility = EVisibility::Collapsed;
 };
