Fix: Session Frontend - test log details scrollbar changes justification when scrolling

diff --git a/Engine/Source/Developer/AutomationWindow/Private/SAutomationWindow.cpp b/Engine/Source/Developer/AutomationWindow/Private/SAutomationWindow.cpp
--- a/Engine/Source/Developer/AutomationWindow/Private/SAutomationWindow.cpp
+++ b/Engine/Source/Developer/AutomationWindow/Private/SAutomationWindow.cpp
@@ -600,6 +600,7 @@ void SAutomationWindow::Construct( const FArguments& InArgs, const IAutomationCo
 											SNew(SScrollBox)
 											.Orientation(EOrientation::Orient_Horizontal)
 											+SScrollBox::Slot()
+											.FillSize(1.0f)
 											[
 												SAssignNew(LogListView, SListView<TSharedPtr<FAutomationOutputMessage> >)
 												.ItemHeight(18)
