Add Support for shader model in automation test framework exclusion list

#rnx

diff --git a/Engine/Source/Developer/AutomationWindow/Private/SAutomationTestItem.cpp b/Engine/Source/Developer/AutomationWindow/Private/SAutomationTestItem.cpp
--- a/Engine/Source/Developer/AutomationWindow/Private/SAutomationTestItem.cpp
+++ b/Engine/Source/Developer/AutomationWindow/Private/SAutomationTestItem.cpp
@@ -6,11 +6,14 @@
 #include "IAutomationControllerModule.h"
 #include "SlateOptMacros.h"
 #include "Framework/Application/SlateApplication.h"
+#include "Framework/MultiBox/MultiBoxBuilder.h"
+#include "Widgets/Input/SEditableTextBox.h"
 #include "Widgets/Images/SImage.h"
 #include "Widgets/Images/SSpinningImage.h"
 #include "Widgets/Notifications/SProgressBar.h"
 #include "Widgets/Text/STextBlock.h"
 #include "Widgets/Input/SCheckBox.h"
+#include "SSimpleComboButton.h"
 #include "AutomationWindowStyle.h"
 #include "AutomationTestExcludelist.h"
 #include "SAutomationWindow.h"
@@ -418,14 +421,175 @@ FReply SAutomationTestItem::SetSkipFlag()
 	return FReply::Handled();
 }
 
+template<typename EnumType>
+void GenerateRHIMenuContentFromExcludeOptions(TSharedPtr<FAutomationTestExcludeOptions> Options, FMenuBuilder* MenuBuilder)
+{
+	static const TSet<FName> AllRHI_OptionNames = FAutomationTestExcludeOptions::GetAllRHIOptionNames<EnumType>();
+	for (auto& Item : AllRHI_OptionNames)
+	{
+		TSharedRef<SWidget> FlagWidget =
+			SNew(SCheckBox)
+			.IsChecked(Options->RHIs.Contains(Item))
+			.OnCheckStateChanged_Lambda([Options, Item](ECheckBoxState NewState)
+				{
+					if (NewState == ECheckBoxState::Checked)
+					{
+						Options->RHIs.Add(Item);
+					}
+					else
+					{
+						Options->RHIs.Remove(Item);
+					}
+				})
+			.Padding(FMargin(4.0f, 0.0f))
+			.Content()
+			[
+				SNew(STextBlock)
+				.Text(FText::FromName(Item))
+			];
+
+		MenuBuilder->AddWidget(FlagWidget, FText::GetEmpty());
+	}
+}
+
+TSharedRef<SWidget> GenerateRHIMenuContentFromExcludeOptions(TSharedPtr<FAutomationTestExcludeOptions> Options)
+{
+	FMenuBuilder MenuBuilder(false, nullptr);
+
+	MenuBuilder.BeginSection("AutomationWindow_ExcludeOptions_RHI", LOCTEXT("ExcludeOptions_RHI_Section", "Interfaces"));
+	GenerateRHIMenuContentFromExcludeOptions<ETEST_RHI_Options>(Options, &MenuBuilder);
+	MenuBuilder.EndSection();
+
+	MenuBuilder.BeginSection("AutomationWindow_ExcludeOptions_RHI_FeatureLevel", LOCTEXT("ExcludeOptions_RHI_FeatureLevel_Section", "Feature Levels"));
+	GenerateRHIMenuContentFromExcludeOptions<ETEST_RHI_FeatureLevel_Options>(Options, &MenuBuilder);
+	MenuBuilder.EndSection();
+
+	return MenuBuilder.MakeWidget();
+}
+
+FText GenerateRHITextFromExcludeOptions(TSharedPtr<FAutomationTestExcludeOptions> Options)
+{
+	if (Options->RHIs.Num() == 0)
+	{
+		return LOCTEXT("ExcludeOptions_RHI_All", "All Interfaces");
+	}
+
+	TArray<FString> RHIs;
+	for (auto& Item : Options->RHIs)
+	{
+		RHIs.Add(Item.ToString());
+	}
+
+	return FText::FromString(FString::Join(RHIs, TEXT(", ")));
+}
+
 FReply SAutomationTestItem::OnEditExcludeOptionsClicked()
 {
 #if WITH_EDITOR
 	TSharedPtr<FAutomationTestExcludeOptions> Options = TestStatus->GetExcludeOptions();
-	TSharedPtr<FStructOnScope> StructToDisplay = MakeShareable(new FStructOnScope(FAutomationTestExcludeOptions::StaticStruct(), (uint8*)Options.Get()));
 
-	TSharedRef<SKismetInspector> KismetInspector = SNew(SKismetInspector);
-	KismetInspector->ShowSingleStruct(StructToDisplay);
+	// Define the dialog form.
+	TSharedRef<SWidget> Form = SNew(SBox)
+		.WidthOverride(350)
+		[
+			SNew(SVerticalBox)
+			+ SVerticalBox::Slot()
+			.AutoHeight()
+			.Padding(5.0f)
+			[
+				SNew(SHorizontalBox)
+				+ SHorizontalBox::Slot()
+				.MaxWidth(50)
+				.HAlign(HAlign_Right)
+				.VAlign(VAlign_Center)
+				.Padding(5, 0)
+				[
+					SNew(STextBlock)
+					.Text(LOCTEXT("ExcludeOptions_TestLabel", "Test"))
+				]
+				+ SHorizontalBox::Slot()
+				.FillWidth(3)
+				[
+					SNew(SEditableTextBox)
+					.OverflowPolicy(ETextOverflowPolicy::Ellipsis)
+					.Text(FText::FromName(Options->Test))
+					.ToolTipText(FText::FromName(Options->Test))
+					.IsReadOnly(true)
+				]
+			]
+			+ SVerticalBox::Slot()
+			.AutoHeight()
+			.Padding(5.0f)
+			[
+				SNew(SHorizontalBox)
+				+ SHorizontalBox::Slot()
+				.MaxWidth(50)
+				.HAlign(HAlign_Right)
+				.VAlign(VAlign_Center)
+				.Padding(5, 0)
+				[
+					SNew(STextBlock)
+					.Text(LOCTEXT("ExcludeOptions_Reason", "Reason"))
+				]
+				+ SHorizontalBox::Slot()
+				.FillWidth(3)
+				[
+					SNew(SEditableTextBox)
+					.OverflowPolicy(ETextOverflowPolicy::Ellipsis)
+					.Text(FText::FromName(Options->Reason))
+					.OnTextCommitted_Lambda([this, Options](const FText& NewReason, ETextCommit::Type&) { Options->Reason = FName(NewReason.ToString()); })
+					.ToolTipText(LOCTEXT("ExcludeOptions_Reason_ToolTip", "The reason as to why the test is excluded"))
+				]
+			]
+			+ SVerticalBox::Slot()
+			.AutoHeight()
+			.Padding(5.0f)
+			[
+				SNew(SHorizontalBox)
+				+ SHorizontalBox::Slot()
+				.MaxWidth(50)
+				.HAlign(HAlign_Right)
+				.VAlign(VAlign_Center)
+				.Padding(5, 0)
+				[
+					SNew(STextBlock)
+					.Text(LOCTEXT("ExcludeOptions_RHI", "RHIs"))
+				]
+				+ SHorizontalBox::Slot()
+				.MaxWidth(120)
+				.HAlign(HAlign_Left)
+				.VAlign(VAlign_Center)
+				[
+					SNew(SSimpleComboButton)
+					.OnGetMenuContent_Lambda([Options]() { return GenerateRHIMenuContentFromExcludeOptions(Options); })
+					.Text_Lambda([Options]() { return GenerateRHITextFromExcludeOptions(Options); })
+					.ToolTipText_Lambda([Options]() { return GenerateRHITextFromExcludeOptions(Options); })
+					.HasDownArrow(true)
+				]
+			]
+			+ SVerticalBox::Slot()
+			.AutoHeight()
+			.Padding(5.0f)
+			[
+				SNew(SHorizontalBox)
+				+ SHorizontalBox::Slot()
+				.MaxWidth(50)
+				.HAlign(HAlign_Right)
+				.VAlign(VAlign_Center)
+				.Padding(5, 0)
+				[
+					SNew(STextBlock)
+					.Text(LOCTEXT("ExcludeOptions_Warn", "Warn"))
+				]
+				+ SHorizontalBox::Slot()
+				[
+					SNew(SCheckBox)
+					.ToolTipText(LOCTEXT("ExcludeOptions_Warn_ToolTip", "Raise a warning when skipping this test"))
+					.IsChecked(Options->Warn)
+					.OnCheckStateChanged_Lambda([this, Options](ECheckBoxState NewState) { Options->Warn = NewState == ECheckBoxState::Checked; })
+				]
+			]
+		];
 
 	SGenericDialogWidget::FArguments DialogArguments;
 	DialogArguments.OnOkPressed_Lambda([Options, this]()
@@ -434,7 +598,7 @@ FReply SAutomationTestItem::OnEditExcludeOptionsClicked()
 			TestStatus->SetSkipFlag(true, &Entry, false);
 		});
 
-	SGenericDialogWidget::OpenDialog(LOCTEXT("ExcludeTestOptions", "Exclude Test Options"), KismetInspector, DialogArguments, true);
+	SGenericDialogWidget::OpenDialog(LOCTEXT("ExcludeTestOptions", "Exclude Test Options"), Form, DialogArguments, true);
 #endif
 	return FReply::Handled();
 }
diff --git a/Engine/Source/Runtime/AutomationTest/Private/AutomationTestExcludelist.cpp b/Engine/Source/Runtime/AutomationTest/Private/AutomationTestExcludelist.cpp
--- a/Engine/Source/Runtime/AutomationTest/Private/AutomationTestExcludelist.cpp
+++ b/Engine/Source/Runtime/AutomationTest/Private/AutomationTestExcludelist.cpp
@@ -111,7 +111,7 @@ void UAutomationTestExcludelist::RemoveFromExcludeTest(const FString& TestName)
 	}
 }
 
-bool UAutomationTestExcludelist::IsTestExcluded(const FString& TestName, const FString& RHI, FName* OutReason, bool* OutWarn)
+bool UAutomationTestExcludelist::IsTestExcluded(const FString& TestName, const TSet<FName>& RHI, FName* OutReason, bool* OutWarn)
 {
 	if (auto Entry = GetExcludeTestEntry(TestName, RHI))
 	{
@@ -131,7 +131,7 @@ bool UAutomationTestExcludelist::IsTestExcluded(const FString& TestName, const F
 	return false;
 }
 
-FAutomationTestExcludelistEntry* UAutomationTestExcludelist::GetExcludeTestEntry(const FString& TestName, const FString& RHI)
+FAutomationTestExcludelistEntry* UAutomationTestExcludelist::GetExcludeTestEntry(const FString& TestName, const TSet<FName>& RHI)
 {
 	if (TestName.IsEmpty())
 		return nullptr;
@@ -155,7 +155,8 @@ FAutomationTestExcludelistEntry* UAutomationTestExcludelist::GetExcludeTestEntry
 					OutEntry = &Entry;
 					continue;
 				}
-				if (Entry.RHIs.Contains(*RHI))
+				const int8 IntersectNum = RHI.Intersect(Entry.RHIs).Num();
+				if (IntersectNum > 0 && IntersectNum == Entry.NumRHIType())
 				{
 					return &Entry;
 				}
diff --git a/Engine/Source/Runtime/AutomationTest/Public/AutomationTestExcludelist.h b/Engine/Source/Runtime/AutomationTest/Public/AutomationTestExcludelist.h
--- a/Engine/Source/Runtime/AutomationTest/Public/AutomationTestExcludelist.h
+++ b/Engine/Source/Runtime/AutomationTest/Public/AutomationTestExcludelist.h
@@ -5,41 +5,54 @@
 #include "CoreMinimal.h"
 #include "AutomationTestExcludelist.generated.h"
 
-UENUM(Meta = (Bitflags, UseEnumValuesAsMaskValuesInEditor = "true"))
-enum class ERHI_Flags : uint8
+UENUM()
+enum class ETEST_RHI_Options
 {
-	DirectX11	= 1 << 0 UMETA(DisplayName = "DirectX 11"),
-	DirectX12	= 1 << 1 UMETA(DisplayName = "DirectX 12"),
-	Vulkan		= 1 << 2 UMETA(DisplayName = "Vulkan"),
-	Metal		= 1 << 3 UMETA(DisplayName = "Metal"),
-	NUM			UMETA(Hidden)
+	DirectX11	UMETA(DisplayName = "DirectX 11"),
+	DirectX12	UMETA(DisplayName = "DirectX 12"),
+	Vulkan		UMETA(DisplayName = "Vulkan"),
+	Metal		UMETA(DisplayName = "Metal"),
+	Null		UMETA(DisplayName = "Null")
 };
 
-ENUM_CLASS_FLAGS(ERHI_Flags);
+UENUM()
+enum class ETEST_RHI_FeatureLevel_Options
+{
+	SM5 UMETA(DisplayName = "SM5"),
+	SM6 UMETA(DisplayName = "SM6")
+};
 
 USTRUCT()
 struct FAutomationTestExcludeOptions
 {
 	GENERATED_BODY()
 
-	void SetRHIFlagsFromNames(TSet<FName>& InRHIs)
+	template<typename EnumType>
+	static TSet<FName> GetAllRHIOptionNames()
 	{
-		int32 Num_RHIs = InRHIs.Num();
-		if (Num_RHIs == 0)
-			return;
-
-		int32 InRHIs_index = 0;
-		UEnum* Enum = StaticEnum<ERHI_Flags>();
-		int32 Num_Flags = Enum->NumEnums();
-		for (int32 i = 0; i < Num_Flags && InRHIs_index < Num_RHIs; i++)
+		static TSet<FName> NameSet;
+		if (NameSet.IsEmpty())
 		{
-			if (InRHIs.Contains(*Enum->GetDisplayNameTextByIndex(i).ToString()))
+			if constexpr (std::is_same_v<EnumType, ETEST_RHI_Options> || std::is_same_v<EnumType, ETEST_RHI_FeatureLevel_Options>)
 			{
-				RHIs |= (int8)Enum->GetValueByIndex(i);
-				InRHIs_index++;
+				UEnum* Enum = StaticEnum<EnumType>();
+				int32 Num_Flags = Enum->NumEnums() - 1;
+				for (int32 i = 0; i < Num_Flags; i++)
+				{
+					NameSet.Add(*Enum->GetDisplayNameTextByIndex(i).ToString());
+				}
 			}
 		}
+
+		return NameSet;
+	}
+
+	template<typename EnumType>
+	static FString GetRHIOptionName(EnumType Value)
+	{
+		return StaticEnum<EnumType>()->GetDisplayNameTextByValue((int)Value).ToString();
 	}
+
 	/* Name of the target test */
 	UPROPERTY(VisibleAnywhere, Category = ExcludeTestOptions)
 	FName Test;
@@ -48,9 +61,9 @@ struct FAutomationTestExcludeOptions
 	UPROPERTY(EditAnywhere, Category = ExcludeTestOptions)
 	FName Reason;
 
-	/* Option to target specific RHI. No option means it should be applied to all RHI */
-	UPROPERTY(EditAnywhere, Meta = (Bitmask, BitmaskEnum = "/Script/AutomationTest.ERHI_Flags"), Category = ExcludeTestOptions)
-	int8 RHIs = 0;
+	/* Options to target specific RHI. No option means it should be applied to all RHIs */
+	UPROPERTY(EditAnywhere, Category = ExcludeTestOptions)
+	TSet<FName> RHIs;
 
 	/* Should the Reason be reported as a warning in the log */
 	UPROPERTY(EditAnywhere, Category = ExcludeTestOptions)
@@ -67,34 +80,40 @@ struct FAutomationTestExcludelistEntry
 	FAutomationTestExcludelistEntry(const FAutomationTestExcludeOptions& Options)
 		: Test(Options.Test)
 		, Reason(Options.Reason)
+		, RHIs(Options.RHIs)
 		, Warn(Options.Warn)
-	{
-		UEnum* Enum = StaticEnum<ERHI_Flags>();
-		int32 Num_Flags = Enum->NumEnums();
-		for (int32 i = 0; i < Num_Flags; i++)
-		{
-			int8 Flag = IntCastChecked<int8>(Enum->GetValueByIndex(i));
-			if ((int8)ERHI_Flags::NUM == Flag)
-				break; // We reach the maximum value
-
-			if ((Options.RHIs & Flag) == Flag)
-			{
-				RHIs.Add(*Enum->GetDisplayNameTextByIndex(i).ToString());
-			}
-		}
-	}
+	{ }
 
-	TSharedPtr<FAutomationTestExcludeOptions> GetOptions()
+	TSharedPtr<FAutomationTestExcludeOptions> GetOptions() const
 	{
 		TSharedPtr<FAutomationTestExcludeOptions> Options = MakeShareable(new FAutomationTestExcludeOptions());
 		Options->Test = Test;
 		Options->Reason = Reason;
 		Options->Warn = Warn;
-		Options->SetRHIFlagsFromNames(RHIs);
+		Options->RHIs = RHIs;
 		
 		return Options;
 	}
 
+	/* Return the number of RHI types that needs to be matched for exclusion */
+	int8 NumRHIType() const
+	{
+		int8 Num = 0;
+		// Test mentions of each RHI option types
+		static const TSet<FName> AllRHI_OptionNames = FAutomationTestExcludeOptions::GetAllRHIOptionNames<ETEST_RHI_Options>();
+		if (RHIs.Difference(AllRHI_OptionNames).Num() < RHIs.Num())
+		{
+			Num++;
+		}
+		static const TSet<FName> AllRHI_FeatureLevel_OptionNames = FAutomationTestExcludeOptions::GetAllRHIOptionNames<ETEST_RHI_FeatureLevel_Options>();
+		if (RHIs.Difference(AllRHI_FeatureLevel_OptionNames).Num() < RHIs.Num())
+		{
+			Num++;
+		}
+
+		return Num;
+	}
+
 	/* Determine if exclusion entry is propagated based on test name - used for management in test automation window */
 	void SetPropagation(const FString& ForTestName)
 	{
@@ -181,8 +200,8 @@ public:
 
 	void AddToExcludeTest(const FString& TestName, const FAutomationTestExcludelistEntry& ExcludelistEntry);
 	void RemoveFromExcludeTest(const FString& TestName);
-	bool IsTestExcluded(const FString& TestName, const FString& RHI = TEXT(""), FName* OutReason = nullptr, bool* OutWarn = nullptr);
-	FAutomationTestExcludelistEntry* GetExcludeTestEntry(const FString& TestName, const FString& RHI = TEXT(""));
+	bool IsTestExcluded(const FString& TestName, const TSet<FName>& = TSet<FName>(), FName* OutReason = nullptr, bool* OutWarn = nullptr);
+	FAutomationTestExcludelistEntry* GetExcludeTestEntry(const FString& TestName, const TSet<FName>& = TSet<FName>());
 
 	void SaveConfig();
 	FString GetConfigFilename() { return UObject::GetDefaultConfigFilename(); } const
diff --git a/Engine/Source/Runtime/AutomationWorker/Private/AutomationWorkerModule.cpp b/Engine/Source/Runtime/AutomationWorker/Private/AutomationWorkerModule.cpp
--- a/Engine/Source/Runtime/AutomationWorker/Private/AutomationWorkerModule.cpp
+++ b/Engine/Source/Runtime/AutomationWorker/Private/AutomationWorkerModule.cpp
@@ -13,6 +13,8 @@
 #include "Misc/Paths.h"
 #include "Misc/App.h"
 #include "Modules/ModuleManager.h"
+#include "RHIFeatureLevel.h"
+#include "RHIStrings.h"
 
 #if WITH_ENGINE
 	#include "Engine/Engine.h"
@@ -486,23 +488,28 @@ void FAutomationWorkerModule::HandleScreenShotAndTraceCapturedWithName(const TAr
 }
 #endif
 
-FString GetRHIForAutomation()
+TSet<FName> GetRHIForAutomation()
 {
 	// Remove any extra information in () from RHI string
 	FString RHI = FApp::GetGraphicsRHI();
+	FString FeatureLevel = LexToString(GMaxRHIFeatureLevel);
 	int Pos;
 	if (RHI.FindChar(*TEXT("("), Pos))
 	{
 		RHI = RHI.Left(Pos).TrimEnd();
 	}
-	return RHI;
+	else if (RHI.IsEmpty())
+	{
+		FeatureLevel = RHI = FAutomationTestExcludeOptions::GetRHIOptionName(ETEST_RHI_Options::Null);
+	}
+	return TSet<FName> {FName(RHI), FName(FeatureLevel)};
 }
 
 bool FAutomationWorkerModule::IsTestExcluded(const FString& InTestToRun, FString* OutReason, bool* OutWarn) const
 {
 	FName SkipReason;
 	UAutomationTestExcludelist* Excludelist = UAutomationTestExcludelist::Get();
-	static FString RHI = GetRHIForAutomation();
+	static const TSet<FName> RHI = GetRHIForAutomation();
 	if (Excludelist->IsTestExcluded(InTestToRun, RHI, &SkipReason, OutWarn))
 	{
 		if (OutReason)
