Fix missed skip state in Test Automation window report

diff --git a/Engine/Source/Developer/AutomationWindow/Private/SAutomationGraphicalResultBox.cpp b/Engine/Source/Developer/AutomationWindow/Private/SAutomationGraphicalResultBox.cpp
--- a/Engine/Source/Developer/AutomationWindow/Private/SAutomationGraphicalResultBox.cpp
+++ b/Engine/Source/Developer/AutomationWindow/Private/SAutomationGraphicalResultBox.cpp
@@ -112,6 +112,11 @@ void SAutomationGraphicalResultBox::PopulateData()
 						DeviceIt->TotalTestSuccesses++;
 						ClusterIt->TotalTestSuccesses++;
 					}
+					else if (TestResults.State == EAutomationState::Skipped)
+					{
+						DeviceIt->TotalTestSkips++;
+						ClusterIt->TotalTestSkips++;
+					}
 				}
 			}
 
@@ -155,7 +160,8 @@ void SAutomationGraphicalResultBox::CreateWidgets()
 		FFormatNamedArguments ClusterArgs;
 		ClusterArgs.Add(TEXT("Name"), FText::FromString(ClusterIt->ClusterName));
 		ClusterArgs.Add(TEXT("NumTests"), ClusterIt->TotalNumTests);
-		ClusterArgs.Add(TEXT("NumFails"), ClusterIt->TotalNumTests - ClusterIt->TotalTestSuccesses);
+		ClusterArgs.Add(TEXT("NumSkips"), ClusterIt->TotalTestSkips);
+		ClusterArgs.Add(TEXT("NumFails"), ClusterIt->TotalNumTests - ClusterIt->TotalTestSuccesses - ClusterIt->TotalTestSkips);
 		ClusterArgs.Add(TEXT("TotalTime"), ClusterIt->TotalTime);
 		ClusterArgs.Add(TEXT("ParallelTime"), ClusterIt->ParallelTime);
 
@@ -168,7 +174,7 @@ void SAutomationGraphicalResultBox::CreateWidgets()
 			[
 				SNew(STextBlock)
 				.TextStyle(FAutomationWindowStyle::Get(), "Automation.ReportHeader" )
-				.Text(FText::Format(LOCTEXT("AutomationGraphicalClusterHeader", "{Name}  -  {NumTests} Tests / {NumFails} Fails / {TotalTime} Seconds (Total) / {ParallelTime} Seconds (Parallel)"), ClusterArgs))
+				.Text(FText::Format(LOCTEXT("AutomationGraphicalClusterHeader", "{Name}  -  {NumTests} Tests / {NumFails} Fails / {NumSkips} Skips / {TotalTime} Seconds (Total) / {ParallelTime} Seconds (Parallel)"), ClusterArgs))
 			];
 
 		RowCounter++;
@@ -181,7 +187,8 @@ void SAutomationGraphicalResultBox::CreateWidgets()
 			//Add Device Header
 			FFormatNamedArguments DeviceArgs;
 			DeviceArgs.Add(TEXT("NumTests"), NumTests);
-			DeviceArgs.Add(TEXT("NumFails"), NumTests - DeviceIt->TotalTestSuccesses);
+			DeviceArgs.Add(TEXT("NumSkips"), DeviceIt->TotalTestSkips);
+			DeviceArgs.Add(TEXT("NumFails"), NumTests - DeviceIt->TotalTestSuccesses - DeviceIt->TotalTestSkips);
 			DeviceArgs.Add(TEXT("TotalTime"), DeviceIt->TotalTime);
 
 			GridContainer->AddSlot(0,RowCounter)
@@ -197,7 +204,7 @@ void SAutomationGraphicalResultBox::CreateWidgets()
 					+SVerticalBox::Slot()
 					[
 						SNew(STextBlock)
-						.Text(FText::Format(LOCTEXT("AutomationGraphicalDeviceHeader", "{NumTests} Tests / {NumFails} Fails / {TotalTime} Seconds"), DeviceArgs))
+						.Text(FText::Format(LOCTEXT("AutomationGraphicalDeviceHeader", "{NumTests} Tests / {NumFails} Fails / {NumSkips} Skips / {TotalTime} Seconds"), DeviceArgs))
 					]
 				];
 
diff --git a/Engine/Source/Developer/AutomationWindow/Private/SAutomationGraphicalResultBox.h b/Engine/Source/Developer/AutomationWindow/Private/SAutomationGraphicalResultBox.h
--- a/Engine/Source/Developer/AutomationWindow/Private/SAutomationGraphicalResultBox.h
+++ b/Engine/Source/Developer/AutomationWindow/Private/SAutomationGraphicalResultBox.h
@@ -102,6 +102,7 @@ private:
 		{
 			InstanceName = InName;
 			TotalTime = 0.f;
+			TotalTestSkips = 0;
 			TotalTestSuccesses = 0;
 		}
 
@@ -111,6 +112,9 @@ private:
 		/** Total time of all the tests */
 		float TotalTime;
 
+		/** How many of the tests were skipped */
+		uint32 TotalTestSkips;
+
 		/** How many of the tests were successful */
 		uint32 TotalTestSuccesses;
 
@@ -129,6 +133,7 @@ private:
 			TotalTime = 0.f;
 			ParallelTime = 0.f;
 			TotalNumTests = 0;
+			TotalTestSkips = 0;
 			TotalTestSuccesses = 0;
 		}
 
@@ -138,6 +143,9 @@ private:
 		/** Total number of tests */
 		uint32 TotalNumTests;
 
+		/** How many of the tests were skipped */
+		uint32 TotalTestSkips;
+
 		/** How many of the tests were successful */
 		uint32 TotalTestSuccesses;
 
