Implement auto expansion of subgroups that has only one child in SessionFrontend

diff --git a/Engine/Source/Developer/AutomationController/Private/AutomationControllerSettings.cpp b/Engine/Source/Developer/AutomationController/Private/AutomationControllerSettings.cpp
--- a/Engine/Source/Developer/AutomationController/Private/AutomationControllerSettings.cpp
+++ b/Engine/Source/Developer/AutomationController/Private/AutomationControllerSettings.cpp
@@ -13,6 +13,7 @@ UAutomationControllerSettings::UAutomationControllerSettings(const FObjectInitia
 	, bSuppressLogWarnings(false)
 	, bElevateLogWarningsToErrors(true)
 	, bKeepPIEOpen(false)
+	, bAutoExpandSingleItemSubgroups(true)
 	, bTreatLogWarningsAsTestErrors(true)
 	, GameInstanceLostTimerSeconds(300.0f)
 	, bResetTelemetryStorageOnNewSession(false)
diff --git a/Engine/Source/Developer/AutomationController/Public/AutomationControllerSettings.h b/Engine/Source/Developer/AutomationController/Public/AutomationControllerSettings.h
--- a/Engine/Source/Developer/AutomationController/Public/AutomationControllerSettings.h
+++ b/Engine/Source/Developer/AutomationController/Public/AutomationControllerSettings.h
@@ -97,6 +97,10 @@ public:
 	UPROPERTY(Config)
 	bool bKeepPIEOpen;
 
+	/** Whether to automatically expand Automation Tests tree subgroups that have single non-leaf item as a child (default=true) */
+	UPROPERTY(Config)
+	bool bAutoExpandSingleItemSubgroups;
+
 private:
 	/** Whether to treat log warnings as test errors (default=true) */
 	UPROPERTY(Config, Meta = (DeprecatedProperty, DeprecationMessage = "Use bElevateLogWarningsToErrors instead."))
diff --git a/Engine/Source/Developer/AutomationWindow/Private/SAutomationWindow.cpp b/Engine/Source/Developer/AutomationWindow/Private/SAutomationWindow.cpp
--- a/Engine/Source/Developer/AutomationWindow/Private/SAutomationWindow.cpp
+++ b/Engine/Source/Developer/AutomationWindow/Private/SAutomationWindow.cpp
@@ -106,7 +106,15 @@ public:
 SAutomationWindow::SAutomationWindow() 
 	: ColumnWidth(90.0f)
 	, bIsLabelVisibilityEnabled(false)
+	, bAutoExpandSingleItemSubgroups(true)
 {
+	UAutomationControllerSettings* Settings = UAutomationControllerSettings::StaticClass()->GetDefaultObject<UAutomationControllerSettings>();
+	if (nullptr == Settings)
+	{
+		return;
+	}
+
+	bAutoExpandSingleItemSubgroups = Settings->bAutoExpandSingleItemSubgroups;
 }
 
 SAutomationWindow::~SAutomationWindow()
@@ -207,6 +215,7 @@ void SAutomationWindow::Construct( const FArguments& InArgs, const IAutomationCo
 		.OnSetExpansionRecursive(this, &SAutomationWindow::OnTestExpansionRecursive)
 		//on selection
 		.OnSelectionChanged(this, &SAutomationWindow::OnTestSelectionChanged)
+		.OnExpansionChanged( this, &SAutomationWindow::OnExpansionChanged)
 		// Allow for some spacing between items with a larger item height.
 		.ItemHeight(20.0f)
 #if WITH_EDITOR
@@ -1475,6 +1484,20 @@ TSharedRef< SWidget > SAutomationWindow::GenerateTestsOptionsMenuContent( )
 		];
 
 
+	TSharedRef<SWidget> AutoExpandSingleItemSubgroupsWidget =
+		SNew(SCheckBox)
+		.IsChecked(this, &SAutomationWindow::AutoExpandSingleItemSubgroupsCheckBoxChecked)
+		.OnCheckStateChanged(this, &SAutomationWindow::HandleAutoExpandSingleItemSubgroupsCheckStateChanged)
+		.Padding(FMargin(4.0f, 0.0f))
+		.ToolTipText(LOCTEXT("AutomationAutoExpandSingleItemSubgroupsTip", "If checked, automatic expansion of single-item test subgroups will be enabled"))
+		.IsEnabled(this, &SAutomationWindow::IsAutomationControllerIdle)
+		.Content()
+		[
+			SNew(STextBlock)
+			.Text(LOCTEXT("AutomationAutoExpandSingleItemSubgroupsText", "Auto expand single-item subgroups"))
+		];
+
+
 	MenuBuilder.BeginSection("AutomationWindowRunTest", LOCTEXT("RunTestOptions", "Advanced Settings"));
 	{
 		MenuBuilder.AddWidget(NumTests, FText::GetEmpty());
@@ -1482,6 +1505,7 @@ TSharedRef< SWidget > SAutomationWindow::GenerateTestsOptionsMenuContent( )
 		#if WITH_EDITOR
 		MenuBuilder.AddWidget(KeepPIEOpenWidget, FText::GetEmpty());
 		#endif //WITH_EDITOR
+		MenuBuilder.AddWidget(AutoExpandSingleItemSubgroupsWidget, FText::GetEmpty());
 	}
 	MenuBuilder.EndSection();
 
@@ -1508,6 +1532,18 @@ void SAutomationWindow::HandleKeepPIEOpenBoxCheckStateChanged(ECheckBoxState Che
 	AutomationController->SetKeepPIEOpen(CheckBoxState == ECheckBoxState::Checked);
 }
 
+/** Returns if we should automatically expand single-item test subgroups */
+ECheckBoxState SAutomationWindow::AutoExpandSingleItemSubgroupsCheckBoxChecked() const
+{
+	return bAutoExpandSingleItemSubgroups ? ECheckBoxState::Checked : ECheckBoxState::Unchecked;
+}
+
+/** Toggles if automatic expansion of single-item subgroups is enabled */
+void SAutomationWindow::HandleAutoExpandSingleItemSubgroupsCheckStateChanged(ECheckBoxState CheckBoxState)
+{
+	bAutoExpandSingleItemSubgroups = (CheckBoxState == ECheckBoxState::Checked);
+}
+
 TArray<FString> SAutomationWindow::SaveExpandedTestNames(TSet<TSharedPtr<IAutomationReport>> ExpandedItems)
 {
 	TArray<FString> ExpandedItemsNames;
@@ -1703,6 +1739,46 @@ void SAutomationWindow::PopulateReportSearchStrings( const TSharedPtr< IAutomati
 	OutSearchStrings.Add( Report->GetFullTestPath() );
 }
 
+void SAutomationWindow::OnExpansionChanged(TSharedPtr<IAutomationReport> InItem, bool bExpanded)
+{
+	ExpandSingleItemSubgroups(InItem, bExpanded);
+}
+
+void SAutomationWindow::ExpandSingleItemSubgroups(TSharedPtr<IAutomationReport> InItem, bool bExpanded)
+{
+	check(InItem.IsValid());
+	if (bAutoExpandSingleItemSubgroups && bExpanded)
+	{
+		TArray<TSharedPtr<IAutomationReport>>& FilteredChildren = InItem->GetFilteredChildren();
+		if (FilteredChildren.Num() == 1)
+		{
+			TSharedPtr<IAutomationReport> SingleChild = FilteredChildren.Top();
+			check(SingleChild.IsValid());
+
+			if (!SingleChild->IsParent())
+			{
+				return;
+			}
+
+			TestTable->SetItemExpansion(InItem, bExpanded);
+
+			TArray<TSharedPtr<IAutomationReport>>& SingleChildFilteredChildren = SingleChild->GetFilteredChildren();
+			static auto IsChildPredicate = [](const TSharedPtr<IAutomationReport>& InItemToCheck) -> bool
+			{
+				check(InItemToCheck.IsValid());
+				return (!InItemToCheck->IsParent());
+			};
+
+			const bool bSingleChildHasAtLeastOneChildLeaf = (nullptr != SingleChildFilteredChildren.FindByPredicate(IsChildPredicate));
+			if (bSingleChildHasAtLeastOneChildLeaf)
+			{
+				return;
+			}
+
+			ExpandSingleItemSubgroups(SingleChild, bExpanded);
+		}
+	}
+}
 
 void SAutomationWindow::OnGetChildren(TSharedPtr<IAutomationReport> InItem, TArray<TSharedPtr<IAutomationReport> >& OutItems)
 {
diff --git a/Engine/Source/Developer/AutomationWindow/Private/SAutomationWindow.h b/Engine/Source/Developer/AutomationWindow/Private/SAutomationWindow.h
--- a/Engine/Source/Developer/AutomationWindow/Private/SAutomationWindow.h
+++ b/Engine/Source/Developer/AutomationWindow/Private/SAutomationWindow.h
@@ -250,6 +250,12 @@ private:
 	 */
 	void PopulateReportSearchStrings( const TSharedPtr< IAutomationReport >& Report, OUT TArray< FString >& OutSearchStrings ) const;
 
+	/** Callback for a test expansion changing */
+	void OnExpansionChanged(TSharedPtr<IAutomationReport> InItem, bool bExpanded);
+
+	/** Recursively expands subgroups as a reaction on expansion change (expands single-item subgroups without leafs recursively) */
+	void ExpandSingleItemSubgroups(TSharedPtr<IAutomationReport> InItem, bool bExpanded);
+
 	/** Gets children tests for a node in the hierarchy */
 	void OnGetChildren(TSharedPtr<IAutomationReport> InItem, TArray<TSharedPtr<IAutomationReport> >& OutItems);
 
@@ -344,6 +350,12 @@ private:
 	/** Toggles if PIE should be kept open when test pass end */
 	void HandleKeepPIEOpenBoxCheckStateChanged(ECheckBoxState CheckBoxState);
 
+	/** Returns if we should automatically expand single-item test subgroups */
+	ECheckBoxState AutoExpandSingleItemSubgroupsCheckBoxChecked() const;
+
+	/** Toggles if automatic expansion of single-item subgroups is enabled */
+	void HandleAutoExpandSingleItemSubgroupsCheckStateChanged(ECheckBoxState CheckBoxState);
+
 	/** Returns if a device group is enabled */
 	ECheckBoxState IsDeviceGroupCheckBoxIsChecked(const int32 DeviceGroupFlag) const;
 	
@@ -638,6 +650,9 @@ private:
 	/** Flag to enable display of the text labels in the toolbar. */
 	bool bIsLabelVisibilityEnabled;
 
+	/** Flag to enable auto expand of subgroups that have single-item groups without children. */
+	bool bAutoExpandSingleItemSubgroups;
+
 	/** Holds a pointer to the preset manager. */
 	TSharedPtr<FAutomationTestPresetManager> TestPresetManager;
 
