Fix filtered out test results in automation window for the case when "Search" bar is not empty.
Introduce GetEnabledReports and GetFilteredReports; deprecated unspecific GetReports

#rnx

diff --git a/Engine/Source/Developer/AutomationController/Private/AutomationControllerManager.cpp b/Engine/Source/Developer/AutomationController/Private/AutomationControllerManager.cpp
--- a/Engine/Source/Developer/AutomationController/Private/AutomationControllerManager.cpp
+++ b/Engine/Source/Developer/AutomationController/Private/AutomationControllerManager.cpp
@@ -939,7 +939,7 @@ void FAutomationControllerManager::ProcessResults()
 	bHasWarning = false;
 	bHasLogs = false;
 
-	TArray< TSharedPtr< IAutomationReport > >& TestReports = GetReports();
+	TArray< TSharedPtr< IAutomationReport > > TestReports = GetEnabledReports();
 
 	if ( TestReports.Num() )
 	{
diff --git a/Engine/Source/Developer/AutomationController/Private/AutomationControllerManager.h b/Engine/Source/Developer/AutomationController/Private/AutomationControllerManager.h
--- a/Engine/Source/Developer/AutomationController/Private/AutomationControllerManager.h
+++ b/Engine/Source/Developer/AutomationController/Private/AutomationControllerManager.h
@@ -244,11 +244,22 @@ public:
 		ReportManager.SetFilter( InFilter );
 	}
 
+	UE_DEPRECATED(5.3, "Use GetFilteredReports or GetEnabledReports instead.")
 	virtual TArray <TSharedPtr <IAutomationReport> >& GetReports() override
+	{
+		return GetFilteredReports();
+	}
+
+	virtual TArray <TSharedPtr <IAutomationReport> >& GetFilteredReports() override
 	{
 		return ReportManager.GetFilteredReports();
 	}
 
+	virtual TArray <TSharedPtr <IAutomationReport> > GetEnabledReports() override
+	{
+		return ReportManager.GetEnabledTestReports();
+	}
+
 	virtual int32 GetNumDeviceClusters() const override
 	{
 		return DeviceClusterManager.GetNumClusters();
diff --git a/Engine/Source/Developer/AutomationController/Public/IAutomationControllerManager.h b/Engine/Source/Developer/AutomationController/Public/IAutomationControllerManager.h
--- a/Engine/Source/Developer/AutomationController/Public/IAutomationControllerManager.h
+++ b/Engine/Source/Developer/AutomationController/Public/IAutomationControllerManager.h
@@ -239,9 +239,20 @@ public:
 	virtual void SetFilter( TSharedPtr< AutomationFilterCollection > InFilter ) = 0;
 
 	/**
-	 * Gives the array of test results to the UI.
+	 * Gives the array of filtered test results to the UI.
 	 */
+	UE_DEPRECATED(5.3, "Use GetFilteredReports or GetEnabledReports instead.")
 	virtual TArray <TSharedPtr <IAutomationReport> >& GetReports() = 0;
+	
+	/**
+	 * Gives the array of filtered test results to the UI.
+	 */
+	virtual TArray <TSharedPtr <IAutomationReport> >& GetFilteredReports() = 0;
+
+	/**
+	 * Gives the array of enabled test results to the UI.
+	 */
+	virtual TArray <TSharedPtr <IAutomationReport> > GetEnabledReports() = 0;
 
 	/**
 	 * Get num devices types.
diff --git a/Engine/Source/Developer/AutomationWindow/Private/SAutomationGraphicalResultBox.cpp b/Engine/Source/Developer/AutomationWindow/Private/SAutomationGraphicalResultBox.cpp
--- a/Engine/Source/Developer/AutomationWindow/Private/SAutomationGraphicalResultBox.cpp
+++ b/Engine/Source/Developer/AutomationWindow/Private/SAutomationGraphicalResultBox.cpp
@@ -72,7 +72,7 @@ void SAutomationGraphicalResultBox::PopulateData()
 {
 	//Get the report list
 	TArray< TSharedPtr< IAutomationReport > > AllReports;
-	AllReports = AutomationController->GetReports();
+	AllReports = AutomationController->GetEnabledReports();
 
 	//Find only the enabled tests
 	TArray< TSharedPtr< IAutomationReport > > EnabledReports;
diff --git a/Engine/Source/Developer/AutomationWindow/Private/SAutomationWindow.cpp b/Engine/Source/Developer/AutomationWindow/Private/SAutomationWindow.cpp
--- a/Engine/Source/Developer/AutomationWindow/Private/SAutomationWindow.cpp
+++ b/Engine/Source/Developer/AutomationWindow/Private/SAutomationWindow.cpp
@@ -188,7 +188,7 @@ void SAutomationWindow::Construct( const FArguments& InArgs, const IAutomationCo
 	AutomationControllerState = AutomationController->GetTestState();
 	
 	//cache off reference to filtered reports
-	TArray <TSharedPtr <IAutomationReport> >& TestReports = AutomationController->GetReports();
+	TArray <TSharedPtr <IAutomationReport> >& TestReports = AutomationController->GetFilteredReports();
 
 	// Create the search filter and set criteria
 	AutomationTextFilter = MakeShareable( new AutomationReportTextFilter( AutomationReportTextFilter::FItemToStringArray::CreateSP( this, &SAutomationWindow::PopulateReportSearchStrings ) ) );
@@ -1401,7 +1401,7 @@ void SAutomationWindow::HandlePresetCheckStateChanged(ECheckBoxState CheckBoxSta
 
 		//Expand selected items
 		TestTable->ClearExpandedItems();
-		TArray< TSharedPtr< IAutomationReport > >& TestReports = AutomationController->GetReports();
+		TArray< TSharedPtr< IAutomationReport > >& TestReports = AutomationController->GetFilteredReports();
 		for (int32 Index = 0; Index < TestReports.Num(); Index++)
 		{
 			ExpandEnabledTests(TestReports[Index]);
@@ -1658,7 +1658,7 @@ void SAutomationWindow::FindTestReportsForCurrentEditorLevel(TArray<TSharedPtr<I
 		const FString MapPath = GWorld->GetCurrentLevel()->GetPackage()->GetPathName();
 		if (!MapPath.IsEmpty())
 		{
-			auto FunctionTestsReport = GetFunctionalTestsReport(AutomationController->GetReports());
+			auto FunctionTestsReport = GetFunctionalTestsReport(AutomationController->GetFilteredReports());
 			if (FunctionTestsReport.IsValid())
 			{
 				FindReportByGameRelativeAssetPath(FunctionTestsReport, MapPath, OutLevelReports);
@@ -1692,7 +1692,7 @@ void SAutomationWindow::OnRunLevelTest()
 
 void SAutomationWindow::ScrollToTest(TSharedPtr<IAutomationReport> InReport)
 {
-	auto& RootReports = AutomationController->GetReports();
+	auto& RootReports = AutomationController->GetFilteredReports();
 	for ( auto ChildReport : RootReports )
 	{
 		auto ShouldExpand = ExpandToTest(ChildReport, InReport);
@@ -2074,7 +2074,7 @@ void SAutomationWindow::OnRefreshTestCallback()
 	// Only expand the child nodes if we have a text filter
 	bool ExpandChildren = !AutomationTextFilter->GetRawFilterText().IsEmpty();
 
-	TArray< TSharedPtr< IAutomationReport > >& TestReports = AutomationController->GetReports();
+	TArray< TSharedPtr< IAutomationReport > >& TestReports = AutomationController->GetFilteredReports();
 
 	for( int32 Index = 0; Index < TestReports.Num(); Index++ )
 	{
@@ -2102,7 +2102,7 @@ void SAutomationWindow::OnTestAvailableCallback( EAutomationControllerModuleStat
 	AutomationControllerState = InAutomationControllerState;
 
 	// Only list tests on opening the Window if the asset registry isn't in the middle of loading tests.
-	if ( InAutomationControllerState == EAutomationControllerModuleState::Ready && AutomationController->GetReports().Num() == 0 && !bIsRequestingTests)
+	if ( InAutomationControllerState == EAutomationControllerModuleState::Ready && AutomationController->GetFilteredReports().Num() == 0 && !bIsRequestingTests)
 	{
 #if WITH_EDITOR
 		FAssetRegistryModule& AssetRegistryModule = FModuleManager::LoadModuleChecked<FAssetRegistryModule>(TEXT("AssetRegistry"));
