Fix hyper link functionality in Automation Test Results widget
+ make sure filename token is always placed at the end of the string

#rnx

diff --git a/Engine/Source/Developer/AutomationWindow/Private/SAutomationWindow.cpp b/Engine/Source/Developer/AutomationWindow/Private/SAutomationWindow.cpp
--- a/Engine/Source/Developer/AutomationWindow/Private/SAutomationWindow.cpp
+++ b/Engine/Source/Developer/AutomationWindow/Private/SAutomationWindow.cpp
@@ -1991,9 +1991,9 @@ TSharedRef<ITableRow> SAutomationWindow::OnGenerateWidgetForLog(TSharedPtr<FAuto
 {
 	check(Message.IsValid());
 
-	// ^((?:[\w]\:|\\)(?:(?:\\[a-z_\-\s0-9\.]+)+)\.(?:cpp|h))\((\d+)\)
-	// https://regex101.com/r/vV4cV7/1
-	FRegexPattern FileAndLinePattern(TEXT("^((?:[\\w]\\:|\\\\)(?:(?:\\\\[a-z_\\-\\s0-9\\.]+)+)\\.(?:cpp|h))\\((\\d+)\\)"));
+	// \[((?:[\w]\:|\\)(?:(?:\\[A-Za-z_\-\s0-9\.]+)+)\.(?:cpp|h))\((\d+)\)\]$
+	// https://regex101.com/r/vV4cV7/25
+	FRegexPattern FileAndLinePattern(TEXT("\\[((?:[\\w]\\:|\\\\)(?:(?:\\\\[A-Za-z_\\-\\s0-9\\.]+)+)\\.(?:cpp|h))\\((\\d+)\\)\\]$"));
 	FRegexMatcher FileAndLineRegexMatcher(FileAndLinePattern, Message->Text);
 
 	TSharedRef<SWidget> SourceLink = SNullWidget::NullWidget;
@@ -2006,7 +2006,7 @@ TSharedRef<ITableRow> SAutomationWindow::OnGenerateWidgetForLog(TSharedPtr<FAuto
 		int32 LineNumber = FCString::Atoi(*FileAndLineRegexMatcher.GetCaptureGroup(2));
 
 		// Remove the hyperlink from the message, since we're splitting it into its own string.
-		MessageString.RightChopInline(FileAndLineRegexMatcher.GetMatchEnding(), false);
+		MessageString.LeftChopInline(FileAndLineRegexMatcher.GetCaptureGroup(0).Len(), false);
 
 		SourceLink = SNew(SHyperlink)
 			.Style(FAutomationWindowStyle::Get(), "Common.GotoNativeCodeHyperlink")
@@ -2023,16 +2023,16 @@ TSharedRef<ITableRow> SAutomationWindow::OnGenerateWidgetForLog(TSharedPtr<FAuto
 			.AutoWidth()
 			.Padding(0)
 			[
-				SourceLink
+				SNew(STextBlock)
+				.TextStyle(FAutomationWindowStyle::Get(), Message->Style )
+				.Text(FText::FromString(MessageString))
 			]
 
 			+ SHorizontalBox::Slot()
 			.AutoWidth()
 			.Padding(0)
 			[
-				SNew(STextBlock)
-				.TextStyle(FAutomationWindowStyle::Get(), Message->Style )
-				.Text(FText::FromString(MessageString))
+				SourceLink
 			]
 		];
 }
diff --git a/Engine/Source/Runtime/Core/Private/Misc/AutomationTest.cpp b/Engine/Source/Runtime/Core/Private/Misc/AutomationTest.cpp
--- a/Engine/Source/Runtime/Core/Private/Misc/AutomationTest.cpp
+++ b/Engine/Source/Runtime/Core/Private/Misc/AutomationTest.cpp
@@ -1011,7 +1011,16 @@ FString FAutomationExecutionEntry::ToString() const
 	FString ComplexString;
 
 	ComplexString = Event.Message;
-	
+
+	if (!Event.Context.IsEmpty())
+	{
+		ComplexString += TEXT(" [");
+		ComplexString += Event.Context;
+		ComplexString += TEXT("] ");
+	}
+
+	// Place the filename at the end so it can be extracted by the SAutomationWindow widget
+	// Expectation is "[filename(line)]"
 	if ( !Filename.IsEmpty() && LineNumber > 0 )
 	{
 		ComplexString += TEXT(" [");
@@ -1021,13 +1030,6 @@ FString FAutomationExecutionEntry::ToString() const
 		ComplexString += TEXT(")]");
 	}
 
-	if ( !Event.Context.IsEmpty() )
-	{
-		ComplexString += TEXT(" [");
-		ComplexString += Event.Context;
-		ComplexString += TEXT("] ");
-	}
-
 	return ComplexString;
 }
 
