Add visual cue when hover the Exclusion icon to let the user know it can be interact with

diff --git a/Engine/Source/Developer/AutomationWindow/Private/SAutomationTestItem.cpp b/Engine/Source/Developer/AutomationWindow/Private/SAutomationTestItem.cpp
--- a/Engine/Source/Developer/AutomationWindow/Private/SAutomationTestItem.cpp
+++ b/Engine/Source/Developer/AutomationWindow/Private/SAutomationTestItem.cpp
@@ -75,10 +75,16 @@ TSharedRef<SWidget> SAutomationTestItem::GenerateWidgetForColumn( const FName& C
 				.ButtonStyle(FAutomationWindowStyle::Get(), "NoBorder")
 				.ToolTipText(this, &SAutomationTestItem::GetExcludeReason)
 				.OnClicked(FOnClicked::CreateSP(this, &SAutomationTestItem::SetSkipFlag))
+#if WITH_EDITOR
+				.Cursor_Lambda([this]() {return this->IsLocalSession ? EMouseCursor::Hand : EMouseCursor::Default;})
+				.OnHovered_Lambda([this]() {this->bIsToBeSkippedButtonHovered = true;})
+				.OnUnhovered_Lambda([this]() {this->bIsToBeSkippedButtonHovered = false;})
+#endif
 				[
 					SNew(SImage)
 					.Image(FAutomationWindowStyle::Get().GetBrush("AutomationWindow.ExcludedTestsFilter"))
 					.Visibility(this, &SAutomationTestItem::IsToBeSkipped_GetVisibility)
+					.ColorAndOpacity(this, &SAutomationTestItem::IsToBeSkipped_GetColorAndOpacity)
 				]
 			];
 	}
@@ -92,6 +98,7 @@ TSharedRef<SWidget> SAutomationTestItem::GenerateWidgetForColumn( const FName& C
 			[
 				SNew(SSimpleButton)
 				.Icon(FAutomationWindowStyle::Get().GetBrush("Icons.Edit"))
+				.Cursor_Lambda([this]() {return this->IsLocalSession ? EMouseCursor::Hand : EMouseCursor::Default;})
 				.Visibility(this, &SAutomationTestItem::IsDirectlyExcluded_GetVisibility)
 				.ToolTipText(LOCTEXT("EditExcludeOptions", "Edit exclude options"))
 				.OnClicked(FOnClicked::CreateSP(this, &SAutomationTestItem::OnEditExcludeOptionsClicked))
@@ -385,9 +392,22 @@ ECheckBoxState SAutomationTestItem::IsTestEnabled() const
 
 EVisibility SAutomationTestItem::IsToBeSkipped_GetVisibility() const
 {
+	if (bIsToBeSkippedButtonHovered && IsLocalSession)
+	{
+		return EVisibility::Visible;
+	}
 	return TestStatus->IsToBeSkipped() ? EVisibility::Visible : EVisibility::Hidden;
 }
 
+FSlateColor SAutomationTestItem::IsToBeSkipped_GetColorAndOpacity() const
+{
+	if (bIsToBeSkippedButtonHovered && !TestStatus->IsToBeSkipped())
+	{
+		return FLinearColor(1.0f, 1.0f, 1.0f, 0.4f);
+	}
+	return FLinearColor(1.0f, 1.0f, 1.0f, 1.0f);
+}
+
 bool SAutomationTestItem::IsDirectlyExcluded() const
 {
 	return WITH_EDITOR && IsLocalSession && TestStatus->IsToBeSkipped() && !TestStatus->IsToBeSkippedByPropagation();
@@ -403,12 +423,22 @@ FText SAutomationTestItem::GetExcludeReason() const
 	FName Reason;
 	bool IsToBeSkipped = TestStatus->IsToBeSkipped(&Reason);
 
-	return IsToBeSkipped ? FText::FromName(Reason) : FText();
+	if (IsToBeSkipped)
+	{
+		return FText::FromName(Reason);
+	}
+
+	return IsLocalSession ? LOCTEXT("ExludeTest", "Exclude test") : FText();
 }
 
 FReply SAutomationTestItem::SetSkipFlag()
 {
 #if WITH_EDITOR
+	// If it's not local session editing is disabled
+	if (IsLocalSession == false)
+	{
+		return FReply::Handled();
+	}
 	if (!TestStatus->IsToBeSkipped())
 	{
 		OnEditExcludeOptionsClicked();
diff --git a/Engine/Source/Developer/AutomationWindow/Private/SAutomationTestItem.h b/Engine/Source/Developer/AutomationWindow/Private/SAutomationTestItem.h
--- a/Engine/Source/Developer/AutomationWindow/Private/SAutomationTestItem.h
+++ b/Engine/Source/Developer/AutomationWindow/Private/SAutomationTestItem.h
@@ -88,6 +88,13 @@ protected:
 	*/
 	EVisibility IsToBeSkipped_GetVisibility() const;
 
+	/**
+	* Change opacity of skipped button if it is hovered.
+	*
+	* Return the Linear Color with different opacity based on hovered status.
+	*/
+	FSlateColor IsToBeSkipped_GetColorAndOpacity() const;
+
 	/**
 	* Is the test inside exclude list through direct exclusion (not through propagation).
 	*
@@ -206,6 +213,9 @@ private:
 	/** Is Session local. */
 	bool IsLocalSession;
 
+	/** Is Skip button hovered. */
+	bool bIsToBeSkippedButtonHovered = false;
+
 	/** Holds the highlight string for the automation test. */
 	TAttribute<FText> HighlightText;
 
