AutomationTest - fix search text highlighting on leaves in Test Automation window

#rnx

diff --git a/Engine/Source/Developer/AutomationWindow/Private/SAutomationTestItem.cpp b/Engine/Source/Developer/AutomationWindow/Private/SAutomationTestItem.cpp
--- a/Engine/Source/Developer/AutomationWindow/Private/SAutomationTestItem.cpp
+++ b/Engine/Source/Developer/AutomationWindow/Private/SAutomationTestItem.cpp
@@ -121,6 +121,7 @@ TSharedRef<SWidget> SAutomationTestItem::GenerateWidgetForColumn( const FName& C
 				.OnNavigate_Lambda([this] {
 					GEngine->Exec(nullptr, *TestStatus->GetOpenCommand());
 				})
+				.HighlightText(HighlightText)
 				.Text(FText::FromString(TestStatus->GetDisplayNameWithDecoration()));
 #endif
 		}
@@ -145,6 +146,7 @@ TSharedRef<SWidget> SAutomationTestItem::GenerateWidgetForColumn( const FName& C
 						}
 					}
 				})
+				.HighlightText(HighlightText)
 				.Text(FText::FromString(TestStatus->GetDisplayNameWithDecoration()));
 #endif
 		}
@@ -153,6 +155,7 @@ TSharedRef<SWidget> SAutomationTestItem::GenerateWidgetForColumn( const FName& C
 			TestNameWidget = SNew(SHyperlink)
 				.Style(FAutomationWindowStyle::Get(), "Common.GotoNativeCodeHyperlink")
 				.OnNavigate_Lambda([this] { FSlateApplication::Get().GotoLineInSource(TestStatus->GetSourceFile(), TestStatus->GetSourceFileLine()); })
+				.HighlightText(HighlightText)
 				.Text(FText::FromString(TestStatus->GetDisplayNameWithDecoration()));
 		}
 		else
