AutomationTest - Add support for a task tracker field in the Exclude Test Options

#rnx

diff --git a/Engine/Source/Developer/AutomationWindow/Private/SAutomationTestItem.cpp b/Engine/Source/Developer/AutomationWindow/Private/SAutomationTestItem.cpp
--- a/Engine/Source/Developer/AutomationWindow/Private/SAutomationTestItem.cpp
+++ b/Engine/Source/Developer/AutomationWindow/Private/SAutomationTestItem.cpp
@@ -26,6 +26,7 @@
 	#include "AssetRegistry/AssetRegistryModule.h"
 	#include "Dialogs/Dialogs.h"
 	#include "SKismetInspector.h"
+	#include "Internationalization/Regex.h"
 #endif
 
 #include "Widgets/Input/SHyperlink.h"
@@ -33,6 +34,101 @@
 
 #define LOCTEXT_NAMESPACE "AutomationTestItem"
 
+namespace
+{
+#if WITH_EDITOR
+	struct FTaskStringEntry
+	{
+	public:
+		/** Stores found full ticket string including hashtag. */
+		FString FullTicketString = {};
+		/** Stores found ticket identifier string. */
+		FString TicketIdString = {};
+		/** Stores starting position (in original string) where FullTicketString starts. */
+		int32 FullTicketBeginning = INDEX_NONE;
+		/** Stores starting position (in original string) where TicketIdString starts. */
+		int32 TicketIdBeginning = INDEX_NONE;
+
+		/**
+		 * Checks if the current instance stores valid data about task tracker's ticket.
+		 *
+		 * @return Return true if the instance is valid or false otherwise.
+		 */
+		bool IsValid() const
+		{
+			return !FullTicketString.IsEmpty() && (INDEX_NONE != FullTicketBeginning);
+		}
+
+		/**
+		 * Parses the given string in order to find the last occurence of a task tracker's ticket.
+		 * 
+		 * @param Source The string to be parsed.
+		 *
+		 * @return Return valid FTaskStringEntry instance if the ticket information is found successfully or invalid object otherwise.
+		 */
+		static FTaskStringEntry LocateLastEntry(const FString& Source)
+		{
+			FTaskStringEntry ResultData;
+
+			int32 LastEntryIndex = INDEX_NONE;
+			{
+				// First pass is to find the proper index of the last entry.
+				FRegexMatcher TicketRegexMatcher(GetTaskTrackerTicketRegexPattern(), Source);
+
+				while (TicketRegexMatcher.FindNext())
+				{
+					++LastEntryIndex;
+				}
+			}
+
+			{
+				// Second pass is to fill the ResultData
+				FRegexMatcher TicketRegexMatcher(GetTaskTrackerTicketRegexPattern(), Source);
+
+				for (int32 CurrentEntryIndex = 0; CurrentEntryIndex <= LastEntryIndex; ++CurrentEntryIndex)
+				{
+					TicketRegexMatcher.FindNext();
+				}
+
+				ResultData.FullTicketString = TicketRegexMatcher.GetCaptureGroup(0);
+				ResultData.FullTicketBeginning = TicketRegexMatcher.GetCaptureGroupBeginning(0);
+				// Note that the ticket id string might be empty
+				ResultData.TicketIdString = TicketRegexMatcher.GetCaptureGroup(3);
+				ResultData.TicketIdBeginning = TicketRegexMatcher.GetCaptureGroupBeginning(3);
+			}
+
+			return ResultData;
+		}
+
+	private:
+		static FRegexPattern GetTaskTrackerTicketRegexPattern()
+		{
+			static const UAutomationTestExcludelist* Excludelist = UAutomationTestExcludelist::Get();
+			check(nullptr != Excludelist);
+
+			return FRegexPattern(FString::Format(
+				TEXT("\\B({0})(\\s+([A-Za-z0-9\\-_]+))?\\b"),
+				{
+					Excludelist->GetTaskTrackerTicketTag()
+				}));
+		}
+	};
+
+	FString BuildTaskTrackerTicketURL(const FString& TaskTrackerURLBase, const FString& TaskTrackerTicketId)
+	{
+		check((!TaskTrackerURLBase.IsEmpty()) && (!TaskTrackerTicketId.IsEmpty()));
+
+		const TMap<FString, FStringFormatArg> Args =
+		{
+			{ TEXT("ID"), TaskTrackerTicketId }
+		};
+
+		return FString::Format(*TaskTrackerURLBase, Args);
+	}
+#endif // WITH_EDITOR
+
+} // anonymous namespace
+
 /* SAutomationTestItem interface
  *****************************************************************************/
 
@@ -549,17 +645,43 @@ FReply SAutomationTestItem::OnEditExcludeOptionsClicked()
 {
 #if WITH_EDITOR
 	TSharedPtr<FAutomationTestExcludeOptions> Options = TestStatus->GetExcludeOptions();
+	check(Options.IsValid());
+
+	static const UAutomationTestExcludelist* Excludelist = UAutomationTestExcludelist::Get();
+	check(nullptr != Excludelist);
+	const FString TaskTrackerURLBase = Excludelist->GetTaskTrackerURLBase();
+	const FString TaskTrackerURLHashtag = Excludelist->GetConfigTaskTrackerHashtag();
+	
+	const bool TaskTrackerSlotIsVisible =
+		!(TaskTrackerURLBase.IsEmpty() && TaskTrackerURLHashtag.IsEmpty());
+	const bool OpenHyperlinkButtonIsEnabled = 
+		(TaskTrackerSlotIsVisible && (!TaskTrackerURLBase.IsEmpty()));
+
+	FString BeautifiedReason = Options->Reason.ToString();
+	FString TaskTrackerTicketId;
+
+	if (TaskTrackerSlotIsVisible)
+	{
+		FTaskStringEntry TaskTrackerTicketEntry = FTaskStringEntry::LocateLastEntry(BeautifiedReason);
+
+		if (TaskTrackerTicketEntry.IsValid() && !TaskTrackerTicketEntry.TicketIdString.IsEmpty())
+		{
+			TaskTrackerTicketId = TaskTrackerTicketEntry.TicketIdString;
+			BeautifiedReason.RemoveAt(TaskTrackerTicketEntry.FullTicketBeginning, TaskTrackerTicketEntry.FullTicketString.Len());
+		}
+	}
 
 	// Define the dialog form.
+	TSharedPtr<SVerticalBox> VBox;
 	TSharedRef<SWidget> Form = SNew(SBox)
 		.WidthOverride(350)
 		[
-			SNew(SVerticalBox)
-			+ SVerticalBox::Slot()
-			.AutoHeight()
-			.Padding(5.0f)
-			[
-				SNew(SHorizontalBox)
+			SAssignNew(VBox, SVerticalBox)
+		];
+
+	VBox->AddSlot().AutoHeight().Padding(5.0f)
+		[
+			SNew(SHorizontalBox)
 				+ SHorizontalBox::Slot()
 				.MaxWidth(55)
 				.HAlign(HAlign_Right)
@@ -567,23 +689,22 @@ FReply SAutomationTestItem::OnEditExcludeOptionsClicked()
 				.Padding(5, 0)
 				[
 					SNew(STextBlock)
-					.Text(LOCTEXT("ExcludeOptions_TestLabel", "Test"))
+						.Text(LOCTEXT("ExcludeOptions_TestLabel", "Test"))
 				]
 				+ SHorizontalBox::Slot()
 				.FillWidth(3)
 				[
 					SNew(SEditableTextBox)
-					.OverflowPolicy(ETextOverflowPolicy::Ellipsis)
-					.Text(FText::FromName(Options->Test))
-					.ToolTipText(FText::FromName(Options->Test))
-					.IsReadOnly(true)
+						.OverflowPolicy(ETextOverflowPolicy::Ellipsis)
+						.Text(FText::FromName(Options->Test))
+						.ToolTipText(FText::FromName(Options->Test))
+						.IsReadOnly(true)
 				]
-			]
-			+ SVerticalBox::Slot()
-			.AutoHeight()
-			.Padding(5.0f)
-			[
-				SNew(SHorizontalBox)
+		];
+
+	VBox->AddSlot().AutoHeight().Padding(5.0f)
+		[
+			SNew(SHorizontalBox)
 				+ SHorizontalBox::Slot()
 				.MaxWidth(55)
 				.HAlign(HAlign_Right)
@@ -591,23 +712,67 @@ FReply SAutomationTestItem::OnEditExcludeOptionsClicked()
 				.Padding(5, 0)
 				[
 					SNew(STextBlock)
-					.Text(LOCTEXT("ExcludeOptions_Reason", "Reason"))
+						.Text(LOCTEXT("ExcludeOptions_Reason", "Reason"))
 				]
 				+ SHorizontalBox::Slot()
 				.FillWidth(3)
 				[
 					SNew(SEditableTextBox)
-					.OverflowPolicy(ETextOverflowPolicy::Ellipsis)
-					.Text(FText::FromName(Options->Reason))
-					.OnTextCommitted_Lambda([this, Options](const FText& NewReason, const ETextCommit::Type&) { Options->Reason = FName(NewReason.ToString()); })
-					.ToolTipText(LOCTEXT("ExcludeOptions_Reason_ToolTip", "The reason as to why the test is excluded"))
+						.OverflowPolicy(ETextOverflowPolicy::Ellipsis)
+						.Text(FText::FromString(BeautifiedReason))
+						.OnTextCommitted_Lambda([Options, &BeautifiedReason, &TaskTrackerTicketId](const FText& NewBeautifiedReason, const ETextCommit::Type&)
+						{
+							BeautifiedReason = NewBeautifiedReason.ToString();
+							Options->UpdateReason(BeautifiedReason, TaskTrackerTicketId);
+						})
+						.ToolTipText(LOCTEXT("ExcludeOptions_Reason_ToolTip", "The reason as to why the test is excluded"))
 				]
-			]
-			+ SVerticalBox::Slot()
-			.AutoHeight()
-			.Padding(5.0f)
+		];
+
+	if (TaskTrackerSlotIsVisible)
+	{
+		VBox->AddSlot().AutoHeight().Padding(5.0f)
 			[
 				SNew(SHorizontalBox)
+					+ SHorizontalBox::Slot()
+					.MaxWidth(55)
+					.HAlign(HAlign_Right)
+					.VAlign(VAlign_Center)
+					.Padding(5, 0)
+					[
+						SNew(STextBlock)
+							.Text(FText::FromString(Excludelist->GetTaskTrackerName()))
+					]
+					+ SHorizontalBox::Slot()
+					.FillWidth(3.5)
+					[
+						SNew(SEditableTextBox)
+							.OverflowPolicy(ETextOverflowPolicy::Ellipsis)
+							.Text(FText::FromString(TaskTrackerTicketId))
+							.ToolTipText(LOCTEXT("ExcludeOptions_TaskTracker_ToolTip", "Task identifier related to the skipping reason"))
+							.OnTextCommitted_Lambda([Options, &BeautifiedReason, &TaskTrackerTicketId](const FText& NewTaskTrackerTicketId, const ETextCommit::Type& CommitType)
+							{
+								TaskTrackerTicketId = NewTaskTrackerTicketId.ToString();
+								Options->UpdateReason(BeautifiedReason, TaskTrackerTicketId);
+							})
+					]
+					+ SHorizontalBox::Slot()
+					.Padding(5, 0)
+					.HAlign(HAlign_Left)
+					.VAlign(VAlign_Center)
+					[
+						SNew(SHyperlink)
+							.Style(FAutomationWindowStyle::Get(), "Common.GotoNativeCodeHyperlink")
+							.IsEnabled_Lambda([OpenHyperlinkButtonIsEnabled, &TaskTrackerTicketId]() { return OpenHyperlinkButtonIsEnabled && (!TaskTrackerTicketId.IsEmpty()); })
+							.OnNavigate_Lambda([&TaskTrackerURLBase, &TaskTrackerTicketId]() { FPlatformProcess::LaunchURL(*(BuildTaskTrackerTicketURL(TaskTrackerURLBase, TaskTrackerTicketId)), nullptr, nullptr); })
+							.Text(FText::FromString(TEXT("Open")))
+					]
+			];
+	}
+
+	VBox->AddSlot().AutoHeight().Padding(5.0f)
+		[
+			SNew(SHorizontalBox)
 				+ SHorizontalBox::Slot()
 				.MaxWidth(55)
 				.HAlign(HAlign_Right)
@@ -615,7 +780,7 @@ FReply SAutomationTestItem::OnEditExcludeOptionsClicked()
 				.Padding(5, 0)
 				[
 					SNew(STextBlock)
-					.Text(LOCTEXT("ExcludeOptions_Platform", "Platforms"))
+						.Text(LOCTEXT("ExcludeOptions_Platform", "Platforms"))
 				]
 				+ SHorizontalBox::Slot()
 				.MaxWidth(200)
@@ -623,22 +788,21 @@ FReply SAutomationTestItem::OnEditExcludeOptionsClicked()
 				.VAlign(VAlign_Center)
 				[
 					SNew(SComboButton)
-					.OnGetMenuContent_Lambda([Options]() { return GeneratePlatformMenuContentFromExcludeOptions(Options); })
-					.HasDownArrow(true)
-					.ButtonContent()
-					[
-						SNew(STextBlock)
-						.Text_Lambda([Options]() { return GeneratePlatformTextFromExcludeOptions(Options); })
-						.ToolTipText_Lambda([Options]() { return GeneratePlatformTextFromExcludeOptions(Options); })
-						.OverflowPolicy(ETextOverflowPolicy::Ellipsis)
-					]
+						.OnGetMenuContent_Lambda([Options]() { return GeneratePlatformMenuContentFromExcludeOptions(Options); })
+						.HasDownArrow(true)
+						.ButtonContent()
+						[
+							SNew(STextBlock)
+								.Text_Lambda([Options]() { return GeneratePlatformTextFromExcludeOptions(Options); })
+								.ToolTipText_Lambda([Options]() { return GeneratePlatformTextFromExcludeOptions(Options); })
+								.OverflowPolicy(ETextOverflowPolicy::Ellipsis)
+						]
 				]
-			]
-			+ SVerticalBox::Slot()
-			.AutoHeight()
-			.Padding(5.0f)
-			[
-				SNew(SHorizontalBox)
+		];
+
+	VBox->AddSlot().AutoHeight().Padding(5.0f)
+		[
+			SNew(SHorizontalBox)
 				+ SHorizontalBox::Slot()
 				.MaxWidth(55)
 				.HAlign(HAlign_Right)
@@ -646,7 +810,7 @@ FReply SAutomationTestItem::OnEditExcludeOptionsClicked()
 				.Padding(5, 0)
 				[
 					SNew(STextBlock)
-					.Text(LOCTEXT("ExcludeOptions_RHI", "RHIs"))
+						.Text(LOCTEXT("ExcludeOptions_RHI", "RHIs"))
 				]
 				+ SHorizontalBox::Slot()
 				.MaxWidth(200)
@@ -654,22 +818,21 @@ FReply SAutomationTestItem::OnEditExcludeOptionsClicked()
 				.VAlign(VAlign_Center)
 				[
 					SNew(SComboButton)
-					.OnGetMenuContent_Lambda([Options]() { return GenerateRHIMenuContentFromExcludeOptions(Options); })
-					.HasDownArrow(true)
-					.ButtonContent()
-					[
-						SNew(STextBlock)
-						.Text_Lambda([Options]() { return GenerateRHITextFromExcludeOptions(Options); })
-						.ToolTipText_Lambda([Options]() { return GenerateRHITextFromExcludeOptions(Options); })
-						.OverflowPolicy(ETextOverflowPolicy::Ellipsis)
-					]
+						.OnGetMenuContent_Lambda([Options]() { return GenerateRHIMenuContentFromExcludeOptions(Options); })
+						.HasDownArrow(true)
+						.ButtonContent()
+						[
+							SNew(STextBlock)
+								.Text_Lambda([Options]() { return GenerateRHITextFromExcludeOptions(Options); })
+								.ToolTipText_Lambda([Options]() { return GenerateRHITextFromExcludeOptions(Options); })
+								.OverflowPolicy(ETextOverflowPolicy::Ellipsis)
+						]
 				]
-			]
-			+ SVerticalBox::Slot()
-			.AutoHeight()
-			.Padding(5.0f)
-			[
-				SNew(SHorizontalBox)
+		];
+
+	VBox->AddSlot().AutoHeight().Padding(5.0f)
+		[
+			SNew(SHorizontalBox)
 				+ SHorizontalBox::Slot()
 				.MaxWidth(55)
 				.HAlign(HAlign_Right)
@@ -677,16 +840,15 @@ FReply SAutomationTestItem::OnEditExcludeOptionsClicked()
 				.Padding(5, 0)
 				[
 					SNew(STextBlock)
-					.Text(LOCTEXT("ExcludeOptions_Warn", "Warn"))
+						.Text(LOCTEXT("ExcludeOptions_Warn", "Warn"))
 				]
 				+ SHorizontalBox::Slot()
 				[
 					SNew(SCheckBox)
-					.ToolTipText(LOCTEXT("ExcludeOptions_Warn_ToolTip", "Raise a warning when skipping this test"))
-					.IsChecked(Options->Warn)
-					.OnCheckStateChanged_Lambda([this, Options](ECheckBoxState NewState) { Options->Warn = NewState == ECheckBoxState::Checked; })
+						.ToolTipText(LOCTEXT("ExcludeOptions_Warn_ToolTip", "Raise a warning when skipping this test"))
+						.IsChecked(Options->Warn)
+						.OnCheckStateChanged_Lambda([this, Options](ECheckBoxState NewState) { Options->Warn = NewState == ECheckBoxState::Checked; })
 				]
-			]
 		];
 
 	SGenericDialogWidget::FArguments DialogArguments;
@@ -697,7 +859,8 @@ FReply SAutomationTestItem::OnEditExcludeOptionsClicked()
 		});
 
 	SGenericDialogWidget::OpenDialog(LOCTEXT("ExcludeTestOptions", "Exclude Test Options"), Form, DialogArguments, true);
-#endif
+#endif // WITH_EDITOR
+	
 	return FReply::Handled();
 }
 
diff --git a/Engine/Source/Runtime/AutomationTest/Private/AutomationTestExcludelist.cpp b/Engine/Source/Runtime/AutomationTest/Private/AutomationTestExcludelist.cpp
--- a/Engine/Source/Runtime/AutomationTest/Private/AutomationTestExcludelist.cpp
+++ b/Engine/Source/Runtime/AutomationTest/Private/AutomationTestExcludelist.cpp
@@ -17,27 +17,70 @@
 
 DEFINE_LOG_CATEGORY_STATIC(LogAutomationTestExcludelist, Log, All);
 
-static const FString FunctionalTestsPreFix = TEXT("Project.Functional Tests.");
-
-void SortExcludelist(TMap<FString, FAutomationTestExcludelistEntry>& List)
+namespace
 {
-	// Sort in alphabetical order, shortest to longest of key property.
-	// That is to naturally gives priority to parent suite over individual test exclusion when calling GetExcludeTestEntry(TestName).
-	List.KeySort([](const FString& A, const FString& B)
-		{
-			return A < B;
-		});
-}
+	const FString FunctionalTestsPreFix = TEXT("Project.Functional Tests.");
+
+	void SortExcludelist(TMap<FString, FAutomationTestExcludelistEntry>& List)
+	{
+		// Sort in alphabetical order, shortest to longest of key property.
+		// That is to naturally gives priority to parent suite over individual test exclusion when calling GetExcludeTestEntry(TestName).
+		List.KeySort([](const FString& A, const FString& B)
+			{
+				return A < B;
+			});
+	}
 
-void SortExcludelist(TArray<FAutomationTestExcludelistEntry>& List)
+	void SortExcludelist(TArray<FAutomationTestExcludelistEntry>& List)
+	{
+		// Sort in alphabetical order, shortest to longest of FullTestName property.
+		// That is to naturally gives priority to parent suite over individual test exclusion when calling GetExcludeTestEntry(TestName).
+		List.Sort([](const FAutomationTestExcludelistEntry& A, const FAutomationTestExcludelistEntry& B)
+			{
+				return A.FullTestName < B.FullTestName;
+			});
+	}
+
+	const FString TicketTrackerURLHashtagPropertyName = TEXT("URLHashtag");
+	const FString TicketTrackerURLBasePropertyName = TEXT("URLBase");
+} // anonymous namespace
+
+#if WITH_EDITOR
+void FAutomationTestExcludeOptions::UpdateReason(const FString& BeautifiedReason, const FString& TaskTrackerTicketId)
 {
-	// Sort in alphabetical order, shortest to longest of FullTestName property.
-	// That is to naturally gives priority to parent suite over individual test exclusion when calling GetExcludeTestEntry(TestName).
-	List.Sort([](const FAutomationTestExcludelistEntry& A, const FAutomationTestExcludelistEntry& B)
+	if (TaskTrackerTicketId.IsEmpty())
+	{
+		Reason = FName(BeautifiedReason);
+	}
+	else
+	{
+		static const UAutomationTestExcludelist* Excludelist = UAutomationTestExcludelist::Get();
+		check(nullptr != Excludelist);
+
+		FString FullTicketString = Excludelist->GetTaskTrackerTicketTag() + TEXT(" ") + TaskTrackerTicketId;
+
+		if (BeautifiedReason.IsEmpty())
+		{
+			Reason = FName(FullTicketString);
+		}
+		else
 		{
-			return A.FullTestName < B.FullTestName;
-		});
+			const bool LastSymbolIsSpaceOrPunct =
+			(
+				TChar<FString::ElementType>::IsWhitespace(BeautifiedReason[BeautifiedReason.Len() - 1])
+				|| TChar<FString::ElementType>::IsPunct(BeautifiedReason[BeautifiedReason.Len() - 1])
+			);
+
+			if (!LastSymbolIsSpaceOrPunct)
+			{
+				FullTicketString = TEXT(" ") + FullTicketString;
+			}
+
+			Reason = FName(BeautifiedReason + FullTicketString);
+		}
+	}
 }
+#endif // WITH_EDITOR
 
 void FAutomationTestExcludelistEntry::Finalize()
 {
@@ -85,11 +128,15 @@ UAutomationTestExcludelist* UAutomationTestExcludelist::Get()
 void UAutomationTestExcludelist::Initialize()
 {
 	DefaultConfig = GetMutableDefault<UAutomationTestExcludelistConfig>();
+	check(nullptr != DefaultConfig);
+
 	if (PlatformConfigs.IsEmpty())
 	{
 		LoadPlatformConfigs();
 		PopulateEntries();
 	}
+
+	DefaultConfig->LoadTaskTrackerProperties();
 }
 
 void UAutomationTestExcludelist::LoadPlatformConfigs()
@@ -221,6 +268,46 @@ FString UAutomationTestExcludelist::GetConfigFilenameForEntry(const FAutomationT
 	return TEXT("");
 }
 
+FString UAutomationTestExcludelist::GetTaskTrackerURLBase() const
+{
+	return DefaultConfig->GetTaskTrackerURLBase();
+}
+
+FString UAutomationTestExcludelist::GetConfigTaskTrackerHashtag() const
+{
+	return DefaultConfig->GetTaskTrackerURLHashtag();
+}
+
+FString UAutomationTestExcludelist::GetBeautifiedTaskTrackerTicketTagSuffix() const
+{
+	static const FString DefaultTaskTrackerTagSuffix = TEXT("unknown");
+
+	FString TaskTrackerTicketTagSuffix = DefaultConfig->GetTaskTrackerURLHashtag();
+	TaskTrackerTicketTagSuffix.TrimStartAndEndInline();
+
+	if (TaskTrackerTicketTagSuffix.IsEmpty())
+	{
+		TaskTrackerTicketTagSuffix = DefaultTaskTrackerTagSuffix;
+	}
+
+	return TaskTrackerTicketTagSuffix;
+}
+
+FString UAutomationTestExcludelist::GetTaskTrackerName() const
+{
+	FString TaskTrackerName = GetBeautifiedTaskTrackerTicketTagSuffix();
+
+	// Capitalize the first letter
+	TaskTrackerName[0] = TChar<FString::ElementType>::ToUpper(TaskTrackerName[0]);
+
+	return TaskTrackerName;
+}
+
+FString UAutomationTestExcludelist::GetTaskTrackerTicketTag() const
+{
+	return (TEXT("#") + GetBeautifiedTaskTrackerTicketTagSuffix());
+}
+
 void UAutomationTestExcludelist::SaveToConfigs()
 {
 	// Reset the cached configs
@@ -509,4 +596,10 @@ void UAutomationTestExcludelistConfig::SaveConfig()
 	}
 }
 
+void UAutomationTestExcludelistConfig::LoadTaskTrackerProperties()
+{
+	GConfig->GetString(*GetSectionName(), *TicketTrackerURLHashtagPropertyName, TaskTrackerURLHashtag, GEngineIni);
+	GConfig->GetString(*GetSectionName(), *TicketTrackerURLBasePropertyName, TaskTrackerURLBase, GEngineIni);
+}
+
 
diff --git a/Engine/Source/Runtime/AutomationTest/Public/AutomationTestExcludelist.h b/Engine/Source/Runtime/AutomationTest/Public/AutomationTestExcludelist.h
--- a/Engine/Source/Runtime/AutomationTest/Public/AutomationTestExcludelist.h
+++ b/Engine/Source/Runtime/AutomationTest/Public/AutomationTestExcludelist.h
@@ -159,6 +159,9 @@ struct FAutomationTestExcludeOptions
 		return PlatformSettings.FindOrAdd(Platform);
 	}
 
+#if WITH_EDITOR
+	AUTOMATIONTEST_API void UpdateReason(const FString& BeautifiedReason, const FString& TaskTrackerTicketId);
+#endif //WITH_EDITOR
 
 	/* Name of the target test */
 	UPROPERTY(VisibleAnywhere, Category = ExcludeTestOptions)
@@ -325,9 +328,13 @@ public:
 	void Reset();
 	void AddEntry(const FAutomationTestExcludelistEntry& Entry);
 	const TArray<FAutomationTestExcludelistEntry>& GetEntries() const;
+	FString GetTaskTrackerURLHashtag() const { return TaskTrackerURLHashtag; }
+	FString GetTaskTrackerURLBase() const { return TaskTrackerURLBase; }
 
 	void SaveConfig();
 
+	void LoadTaskTrackerProperties();
+
 protected:
 	virtual void PostInitProperties() override;
 	virtual void InitializeSettingsDefault() override;
@@ -336,6 +343,12 @@ protected:
 private:
 	void UpdateHash(const FAutomationTestExcludelistEntry& Entry);
 
+	UPROPERTY(Transient)
+	FString TaskTrackerURLHashtag;
+
+	UPROPERTY(Transient)
+	FString TaskTrackerURLBase;
+
 	UPROPERTY(Config)
 	TArray<FAutomationTestExcludelistEntry> ExcludeTest;
 
@@ -449,6 +462,34 @@ public:
 	*/
 	AUTOMATIONTEST_API FString GetConfigFilenameForEntry(const FAutomationTestExcludelistEntry& Entry, const FName& Platform) const;
 
+	/**
+	 * Get URL base that is used while restoring a task tracker task's URL.
+	 *
+	 * @return Return the URL base.
+	*/
+	AUTOMATIONTEST_API FString GetTaskTrackerURLBase() const;
+
+	/**
+	 * Get non-beautified task tracker's hashtag suffix (from the loaded configs).
+	 *
+	 * @return Return the string that contains non-beautified task tracker's hashtag suffix.
+	 */
+	AUTOMATIONTEST_API FString GetConfigTaskTrackerHashtag() const;
+
+	/**
+	 * Get Task tracker's name if task tracker information is configured.
+	 *
+	 * @return Return task tracker's name if it is configured or empty string otherwise.
+	*/
+	AUTOMATIONTEST_API FString GetTaskTrackerName() const;
+
+	/**
+	 * Get hashtag that is used for referencing tasks if task tracker information is configured.
+	 *
+	 * @return Return hashtag that is used for referencing task tracker tasks if it is configured or empty string otherwise.
+	*/
+	AUTOMATIONTEST_API FString GetTaskTrackerTicketTag() const;
+
 	/** Save all the required configs */
 	AUTOMATIONTEST_API void SaveToConfigs();
 
@@ -465,6 +506,9 @@ private:
 	/** Populate the exclusion entries from the loaded configs */
 	void PopulateEntries();
 
+	/** Get beautified hashtag suffix (starts after leading#) that is used for referencing tasks if task tracker information is configured */
+	FString GetBeautifiedTaskTrackerTicketTagSuffix() const;
+
 	/** Exclusion entries */
 	TMap<FString, FAutomationTestExcludelistEntry> Entries;
 
