Fix Test Automation preset encoding
Use UTF-8 without BOM

#rnx

diff --git a/Engine/Source/Developer/AutomationWindow/Private/AutomationPresetManager.cpp b/Engine/Source/Developer/AutomationWindow/Private/AutomationPresetManager.cpp
--- a/Engine/Source/Developer/AutomationWindow/Private/AutomationPresetManager.cpp
+++ b/Engine/Source/Developer/AutomationWindow/Private/AutomationPresetManager.cpp
@@ -2,12 +2,14 @@
 
 #include "AutomationPresetManager.h"
 #include "HAL/FileManager.h"
+#include "Misc/FileHelper.h"
 #include "UObject/Class.h"
 #include "Serialization/JsonReader.h"
 #include "Serialization/JsonSerializer.h"
 #include "JsonObjectConverter.h"
 
 #include UE_INLINE_GENERATED_CPP_BY_NAME(AutomationPresetManager)
+DEFINE_LOG_CATEGORY_STATIC(LogAutomationPreset, Warning, All);
 
 FAutomationTestPresetManager::FAutomationTestPresetManager()
 {
@@ -47,9 +49,15 @@ TArray<AutomationPresetPtr>& FAutomationTestPresetManager::GetAllPresets()
 
 AutomationPresetPtr FAutomationTestPresetManager::LoadPreset( FArchive& Archive )
 {
-	TSharedPtr<FJsonObject> JsonPreset;
+	TSharedPtr<FJsonObject> JsonPreset = MakeShared<FJsonObject>();
 
-	TSharedRef< TJsonReader<> > JsonReader = TJsonReaderFactory<>::Create(&Archive);
+	FString JsonContent;
+	if ( !FFileHelper::LoadFileToString(JsonContent, Archive) )
+	{
+		return nullptr;
+	}
+
+	TSharedRef< TJsonReader<> > JsonReader = TJsonReaderFactory<>::Create(JsonContent);
 	if ( !FJsonSerializer::Deserialize(JsonReader, JsonPreset) )
 	{
 		return nullptr;
@@ -81,19 +89,13 @@ void FAutomationTestPresetManager::RemovePreset( const AutomationPresetRef Prese
 
 void FAutomationTestPresetManager::SavePreset( const AutomationPresetRef Preset )
 {
-	TSharedPtr<FJsonObject> PresetJson = FJsonObjectConverter::UStructToJsonObject(Preset.Get());
-
-	if ( PresetJson.IsValid() )
+	FString JsonString;
+	if(FJsonObjectConverter::UStructToJsonObjectString(Preset.Get(), JsonString))
 	{
 		FString PresetFileName = GetPresetFolder() / Preset->GetID().ToString() + TEXT(".json");
-		FArchive* PresetFileWriter = IFileManager::Get().CreateFileWriter(*PresetFileName);
-
-		if ( PresetFileWriter != nullptr )
+		if (!FFileHelper::SaveStringToFile(JsonString, *PresetFileName, FFileHelper::EEncodingOptions::ForceUTF8WithoutBOM))
 		{
-			TSharedRef<TJsonWriter<> > JsonWriter = TJsonWriterFactory<>::Create(PresetFileWriter, 0);
-			FJsonSerializer::Serialize(PresetJson.ToSharedRef(), JsonWriter);
-
-			delete PresetFileWriter;
+			UE_LOG(LogAutomationPreset, Error, TEXT("Could not save preset file '%s'"), *Preset->GetID().ToString());
 		}
 	}
 }
@@ -108,7 +110,7 @@ void FAutomationTestPresetManager::LoadPresets()
 	{
 		FString PresetFilePath = GetPresetFolder() / *It;
 		FArchive* PresetFileReader = IFileManager::Get().CreateFileReader(*PresetFilePath);
-		
+
 		if (PresetFileReader != nullptr)
 		{
 			AutomationPresetPtr LoadedPreset = LoadPreset(*PresetFileReader);
@@ -119,7 +121,7 @@ void FAutomationTestPresetManager::LoadPresets()
 			}
 			else
 			{
-				IFileManager::Get().Delete(*PresetFilePath);
+				UE_LOG(LogAutomationPreset, Warning, TEXT("Could not read preset file '%s'. Make sure the file is encoded in UTF-8 without BOM."), *PresetFilePath);
 			}
 
 			delete PresetFileReader;
