Source control support for creating/modifying tests presets in Automation Window.

#rnx

- Submitting code on behalf of partner team.

diff --git a/Engine/Source/Developer/AutomationWindow/AutomationWindow.Build.cs b/Engine/Source/Developer/AutomationWindow/AutomationWindow.Build.cs
--- a/Engine/Source/Developer/AutomationWindow/AutomationWindow.Build.cs
+++ b/Engine/Source/Developer/AutomationWindow/AutomationWindow.Build.cs
@@ -24,6 +24,7 @@ namespace UnrealBuildTool.Rules
                     "CoreUObject",
                     "Json",
                     "JsonUtilities",
+					"SourceControl",
 					"AutomationTest"
 				}
 			);
diff --git a/Engine/Source/Developer/AutomationWindow/Private/AutomationPresetManager.cpp b/Engine/Source/Developer/AutomationWindow/Private/AutomationPresetManager.cpp
--- a/Engine/Source/Developer/AutomationWindow/Private/AutomationPresetManager.cpp
+++ b/Engine/Source/Developer/AutomationWindow/Private/AutomationPresetManager.cpp
@@ -7,6 +7,7 @@
 #include "Serialization/JsonReader.h"
 #include "Serialization/JsonSerializer.h"
 #include "JsonObjectConverter.h"
+#include "SourceControlHelpers.h"
 
 #include UE_INLINE_GENERATED_CPP_BY_NAME(AutomationPresetManager)
 DEFINE_LOG_CATEGORY_STATIC(LogAutomationPreset, Warning, All);
@@ -93,10 +94,34 @@ void FAutomationTestPresetManager::SavePreset( const AutomationPresetRef Preset
 	if(FJsonObjectConverter::UStructToJsonObjectString(Preset.Get(), JsonString))
 	{
 		FString PresetFileName = GetPresetFolder() / Preset->GetID().ToString() + TEXT(".json");
+
+		bool bShouldMarkForAdd = false;
+		if (IFileManager::Get().FileExists(*PresetFileName))
+		{
+			if (USourceControlHelpers::IsEnabled())
+			{
+				// Checkout (or add in case file exists locally but not in source control).
+				USourceControlHelpers::CheckOutOrAddFile(PresetFileName);
+			}
+		}
+		else
+		{
+			bShouldMarkForAdd = true;
+		}
+
 		if (!FFileHelper::SaveStringToFile(JsonString, *PresetFileName, FFileHelper::EEncodingOptions::ForceUTF8WithoutBOM))
 		{
 			UE_LOG(LogAutomationPreset, Error, TEXT("Could not save preset file '%s'"), *Preset->GetID().ToString());
 		}
+
+		if (bShouldMarkForAdd)
+		{
+			if (USourceControlHelpers::IsEnabled())
+			{
+				// Mark for add (or checkout in case file already exists in source control).
+				USourceControlHelpers::CheckOutOrAddFile(PresetFileName);
+			}
+		}
 	}
 }
 
