Prune only log info message + Highlight success with warning in the test automation window

#rnx

diff --git a/Engine/Source/Developer/AutomationWindow/Private/SAutomationTestItem.cpp b/Engine/Source/Developer/AutomationWindow/Private/SAutomationTestItem.cpp
--- a/Engine/Source/Developer/AutomationWindow/Private/SAutomationTestItem.cpp
+++ b/Engine/Source/Developer/AutomationWindow/Private/SAutomationTestItem.cpp
@@ -1076,7 +1076,15 @@ const FSlateBrush* SAutomationTestItem::ItemChildrenStatus_StatusImage(const int
 	if ((TotalComplete > 0) && (CompleteState.TotalEnabled > 0) && TotalComplete == CompleteState.TotalEnabled) {
 		if (TotalComplete == CompleteState.NumEnabledTestsPassed)
 		{
-			ImageToUse = FAutomationWindowStyle::Get().GetBrush("Automation.Success");
+			// If there were ANY warnings in the results
+			if (CompleteState.NumEnabledTestsWarnings > 0 || CompleteState.NumDisabledTestsWarnings > 0)
+			{
+				ImageToUse = FAutomationWindowStyle::Get().GetBrush("Automation.Warning");
+			}
+			else
+			{
+				ImageToUse = FAutomationWindowStyle::Get().GetBrush("Automation.Success");
+			}
 		}
 		else if (CompleteState.NumEnabledTestsFailed)
 		{
diff --git a/Engine/Source/Runtime/AutomationWorker/Private/AutomationWorkerModule.cpp b/Engine/Source/Runtime/AutomationWorker/Private/AutomationWorkerModule.cpp
--- a/Engine/Source/Runtime/AutomationWorker/Private/AutomationWorkerModule.cpp
+++ b/Engine/Source/Runtime/AutomationWorker/Private/AutomationWorkerModule.cpp
@@ -210,7 +210,7 @@ void FAutomationWorkerModule::ReportTestComplete()
 			Message->Duration = ExecutionInfo.Duration;
 			// Prune log entries if it is a success to reduce foot print.
 			Message->Entries = bPruneLogsOnSuccess && bSuccess ?
-				ExecutionInfo.GetEntries().FilterByPredicate([](const FAutomationExecutionEntry& Entry) { return Entry.Event.Context != TEXT("log"); }) : ExecutionInfo.GetEntries();
+				ExecutionInfo.GetEntries().FilterByPredicate([](const FAutomationExecutionEntry& Entry) { return Entry.Event.Context != TEXT("log") || Entry.Event.Type != EAutomationEventType::Info ; }) : ExecutionInfo.GetEntries();
 			Message->WarningTotal = ExecutionInfo.GetWarningTotal();
 			Message->ErrorTotal = ExecutionInfo.GetErrorTotal();
 
