Automation preset save notification

diff --git a/Engine/Source/Developer/AutomationWindow/Private/AutomationPresetManager.cpp b/Engine/Source/Developer/AutomationWindow/Private/AutomationPresetManager.cpp
--- a/Engine/Source/Developer/AutomationWindow/Private/AutomationPresetManager.cpp
+++ b/Engine/Source/Developer/AutomationWindow/Private/AutomationPresetManager.cpp
@@ -88,41 +88,46 @@ void FAutomationTestPresetManager::RemovePreset( const AutomationPresetRef Prese
 	}
 }
 
-void FAutomationTestPresetManager::SavePreset( const AutomationPresetRef Preset )
+bool FAutomationTestPresetManager::SavePreset(const AutomationPresetRef Preset)
 {
 	FString JsonString;
-	if(FJsonObjectConverter::UStructToJsonObjectString(Preset.Get(), JsonString))
+	if (!FJsonObjectConverter::UStructToJsonObjectString(Preset.Get(), JsonString))
 	{
-		FString PresetFileName = GetPresetFolder() / Preset->GetID().ToString() + TEXT(".json");
+		UE_LOG(LogAutomationPreset, Error, TEXT("Could not convert preset file '%s' to JSON string"), *Preset->GetName().ToString());
+		return false;
+	}
 
-		bool bShouldMarkForAdd = false;
-		if (IFileManager::Get().FileExists(*PresetFileName))
-		{
-			if (USourceControlHelpers::IsEnabled())
-			{
-				// Checkout (or add in case file exists locally but not in source control).
-				USourceControlHelpers::CheckOutOrAddFile(PresetFileName);
-			}
-		}
-		else
-		{
-			bShouldMarkForAdd = true;
-		}
+	const FString PresetFileName = GetPresetFolder() / Preset->GetID().ToString() + TEXT(".json");
+	bool bShouldMarkForAdd = false;
 
-		if (!FFileHelper::SaveStringToFile(JsonString, *PresetFileName, FFileHelper::EEncodingOptions::ForceUTF8WithoutBOM))
+	if (IFileManager::Get().FileExists(*PresetFileName))
+	{
+		if (USourceControlHelpers::IsEnabled())
 		{
-			UE_LOG(LogAutomationPreset, Error, TEXT("Could not save preset file '%s'"), *Preset->GetID().ToString());
+			// Checkout (or add in case file exists locally but not in source control).
+			USourceControlHelpers::CheckOutOrAddFile(PresetFileName);
 		}
+	}
+	else
+	{
+		bShouldMarkForAdd = true;
+	}
 
-		if (bShouldMarkForAdd)
+	if (!FFileHelper::SaveStringToFile(JsonString, *PresetFileName, FFileHelper::EEncodingOptions::ForceUTF8WithoutBOM))
+	{
+		UE_LOG(LogAutomationPreset, Error, TEXT("Could not save preset file '%s'"), *Preset->GetID().ToString());
+		return false;
+	}
+
+	if (bShouldMarkForAdd)
+	{
+		if (USourceControlHelpers::IsEnabled())
 		{
-			if (USourceControlHelpers::IsEnabled())
-			{
-				// Mark for add (or checkout in case file already exists in source control).
-				USourceControlHelpers::CheckOutOrAddFile(PresetFileName);
-			}
+			// Mark for add (or checkout in case file already exists in source control).
+			USourceControlHelpers::CheckOutOrAddFile(PresetFileName);
 		}
 	}
+	return true;
 }
 
 void FAutomationTestPresetManager::LoadPresets()
diff --git a/Engine/Source/Developer/AutomationWindow/Private/AutomationPresetManager.h b/Engine/Source/Developer/AutomationWindow/Private/AutomationPresetManager.h
--- a/Engine/Source/Developer/AutomationWindow/Private/AutomationPresetManager.h
+++ b/Engine/Source/Developer/AutomationWindow/Private/AutomationPresetManager.h
@@ -143,7 +143,7 @@ public:
 	 * @param Preset The preset to save.
 	 * @see RemovePreset, LoadPresets
 	 */
-	void SavePreset( const AutomationPresetRef Preset);
+	bool SavePreset( const AutomationPresetRef Preset);
 
 	/**
 	 * Load all Presets from disk.
diff --git a/Engine/Source/Developer/AutomationWindow/Private/SAutomationWindow.cpp b/Engine/Source/Developer/AutomationWindow/Private/SAutomationWindow.cpp
--- a/Engine/Source/Developer/AutomationWindow/Private/SAutomationWindow.cpp
+++ b/Engine/Source/Developer/AutomationWindow/Private/SAutomationWindow.cpp
@@ -45,6 +45,7 @@
 #include "Internationalization/Regex.h"
 #include "AutomationWindowStyle.h"
 #include "AutomationControllerSettings.h"
+#include "Framework/Notifications/NotificationManager.h"
 
 
 #define LOCTEXT_NAMESPACE "AutomationTest"
@@ -1169,12 +1170,28 @@ void SAutomationWindow::ExpandEnabledTests( TSharedPtr< IAutomationReport > InRe
 
 FReply SAutomationWindow::HandleSavePresetClicked()
 {
-	if(SelectedPreset.IsValid())
+	FText PresetSaveStatusText;
+	bool bSavedSuccessfully = false;
+	if (SelectedPreset.IsValid())
 	{
 		TArray<FString> EnabledTests;
 		AutomationController->GetEnabledTestNames(EnabledTests);
 		SelectedPreset->SetEnabledTests(EnabledTests);
-		TestPresetManager->SavePreset(SelectedPreset.ToSharedRef());
+		bSavedSuccessfully = TestPresetManager->SavePreset(SelectedPreset.ToSharedRef());
+		if (bSavedSuccessfully)
+		{
+			PresetSaveStatusText = FText::Format(LOCTEXT("AutomationPresetSaveSucceeded", "Preset '{0}' saved successfully."), SelectedPreset->GetName());
+		}
+		else
+		{
+			PresetSaveStatusText = FText::Format(LOCTEXT("AutomationPresetSaveFailed", "Failed to save preset '{0}'."), SelectedPreset->GetName());
+		}
+	}
+	const FNotificationInfo SaveNotificationInfo = FNotificationInfo(PresetSaveStatusText);
+	TSharedPtr<SNotificationItem> NotificationPtr = FSlateNotificationManager::Get().AddNotification(SaveNotificationInfo);
+	if (NotificationPtr)
+	{
+		NotificationPtr->SetCompletionState(bSavedSuccessfully ? SNotificationItem::CS_Success : SNotificationItem::CS_Fail);
 	}
 	return FReply::Handled();
 }
