Fix a crash when deselecting and the reselecting an instance in the Session manager while a test pass is running

#rnx

diff --git a/Engine/Source/Developer/AutomationWindow/Private/SAutomationWindow.cpp b/Engine/Source/Developer/AutomationWindow/Private/SAutomationWindow.cpp
--- a/Engine/Source/Developer/AutomationWindow/Private/SAutomationWindow.cpp
+++ b/Engine/Source/Developer/AutomationWindow/Private/SAutomationWindow.cpp
@@ -2434,7 +2434,10 @@ void SAutomationWindow::HandleSessionManagerCanSelectSession( const TSharedPtr<I
 
 void SAutomationWindow::HandleSessionManagerSelectionChanged( const TSharedPtr<ISessionInfo>& SelectedSession )
 {
-	FindWorkers();
+	if (!GIsAutomationTesting)
+	{
+		FindWorkers();
+	}
 }
 
 
