Automation Window - Test Tags editor
#rnx

diff --git a/Engine/Source/Developer/AutomationController/Private/AutomationCommandline.cpp b/Engine/Source/Developer/AutomationController/Private/AutomationCommandline.cpp
--- a/Engine/Source/Developer/AutomationController/Private/AutomationCommandline.cpp
+++ b/Engine/Source/Developer/AutomationController/Private/AutomationCommandline.cpp
@@ -264,7 +264,6 @@ public:
 		}
 	}
 
-
 	void ApplyCVarTagFilter(TSharedPtr <AutomationFilterCollection> AutomationFilters)
 	{
 		if (IConsoleVariable* TagCVar = IConsoleManager::Get().FindConsoleVariable(TEXT("Automation.TestTagGlobalFilter")))
@@ -272,6 +271,8 @@ public:
 			FString TagFilterValue = TagCVar->GetString();
 			if (!TagFilterValue.IsEmpty())
 			{
+				FAutomationTestFramework::Get().LoadTestTagMappings();
+
 				UE_LOG(LogAutomationCommandLine, Display, TEXT("Applying Automation Test tag filter '%s'"), *TagFilterValue);
 				FAutomationGroupFilter* FilterTags = new FAutomationGroupFilter();
 				TArray<FAutomatedTestTagFilter> TagList;
@@ -733,6 +734,8 @@ protected:
 				{
 					if (IConsoleVariable* CVar = IConsoleManager::Get().FindConsoleVariable(TEXT("Automation.TestTagGlobalFilter")))
 					{
+						FAutomationTestFramework::Get().LoadTestTagMappings();
+
 						FString FilterValue(TempCmd);
 						Ar.Logf(TEXT("Automation: Setting test tag global filter to: '%s'"), *FilterValue);
 						CVar->Set(*FilterValue);
diff --git a/Engine/Source/Developer/AutomationController/Private/AutomationControllerManager.cpp b/Engine/Source/Developer/AutomationController/Private/AutomationControllerManager.cpp
--- a/Engine/Source/Developer/AutomationController/Private/AutomationControllerManager.cpp
+++ b/Engine/Source/Developer/AutomationController/Private/AutomationControllerManager.cpp
@@ -57,6 +57,7 @@ void FAutomatedTestPassResults::AddTestResult(const IAutomationReportPtr& TestRe
 		TestResult.Test = TestReport;
 		TestResult.TestDisplayName = TestReport->GetDisplayName();
 		TestResult.FullTestPath = FullTestPath;
+		TestResult.Tags = TestReport->GetTagsArray();
 
 		TestsMapIndex.Add(FullTestPath, Tests.Num());
 		Tests.Add(TestResult);
diff --git a/Engine/Source/Developer/AutomationController/Private/AutomationControllerManager.h b/Engine/Source/Developer/AutomationController/Private/AutomationControllerManager.h
--- a/Engine/Source/Developer/AutomationController/Private/AutomationControllerManager.h
+++ b/Engine/Source/Developer/AutomationController/Private/AutomationControllerManager.h
@@ -35,6 +35,9 @@ public:
 	UPROPERTY()
 	FString FullTestPath;
 
+	UPROPERTY()
+	TArray<FString> Tags;
+
 	UPROPERTY()
 	EAutomationState State;
 
diff --git a/Engine/Source/Developer/AutomationController/Private/AutomationReport.cpp b/Engine/Source/Developer/AutomationController/Private/AutomationReport.cpp
--- a/Engine/Source/Developer/AutomationController/Private/AutomationReport.cpp
+++ b/Engine/Source/Developer/AutomationController/Private/AutomationReport.cpp
@@ -48,7 +48,12 @@ FString FAutomationReport::GetTestParameter() const
 
 FString FAutomationReport::GetTags() const
 {
-	return TestInfo.GetTestTags();
+	return FAutomationTestFramework::Get().GetTagsForAutomationTest(this->GetFullTestPath());
+}
+
+TArray<FString> FAutomationReport::GetTagsArray() const
+{
+	return FAutomationTestFramework::Get().GetTagsForAutomationTestAsArray(this->GetFullTestPath());
 }
 
 FString FAutomationReport::GetAssetPath() const
diff --git a/Engine/Source/Developer/AutomationController/Private/AutomationReport.h b/Engine/Source/Developer/AutomationController/Private/AutomationReport.h
--- a/Engine/Source/Developer/AutomationController/Private/AutomationReport.h
+++ b/Engine/Source/Developer/AutomationController/Private/AutomationReport.h
@@ -23,6 +23,7 @@ public:
 	virtual void Empty() override;
 	virtual FString GetTestParameter() const override;
 	virtual FString GetTags() const override;
+	virtual TArray<FString> GetTagsArray() const override;
 	virtual FString GetAssetPath() const override;
 	virtual FString GetOpenCommand() const override;
 	virtual FString GetCommand() const override;
diff --git a/Engine/Source/Developer/AutomationController/Public/IAutomationReport.h b/Engine/Source/Developer/AutomationController/Public/IAutomationReport.h
--- a/Engine/Source/Developer/AutomationController/Public/IAutomationReport.h
+++ b/Engine/Source/Developer/AutomationController/Public/IAutomationReport.h
@@ -241,6 +241,13 @@ public:
 	 */
 	virtual FString GetTags() const = 0;
 
+	/**
+	 * Get the tags associated with this test, as an array.
+	 *
+	 * @return the tags as a string array.
+	 */
+	virtual TArray<FString> GetTagsArray() const = 0;
+
 	/**
 	 * Gets the asset path associated with a test, it may not have one.
 	 * 
diff --git a/Engine/Source/Developer/AutomationWindow/Private/AutomationWindowStyle.cpp b/Engine/Source/Developer/AutomationWindow/Private/AutomationWindowStyle.cpp
--- a/Engine/Source/Developer/AutomationWindow/Private/AutomationWindowStyle.cpp
+++ b/Engine/Source/Developer/AutomationWindow/Private/AutomationWindowStyle.cpp
@@ -98,6 +98,10 @@ FAutomationWindowStyle::FAutomationWindowStyle()
 		Set("AutomationWindow.SmokeTestFilter", new CORE_IMAGE_BRUSH_SVG("Starship/Common/Test", Icon16x16));
 		Set("AutomationWindow.DeveloperDirectoryContent", new CORE_IMAGE_BRUSH_SVG("Starship/Common/Developer", Icon16x16));
 		Set("AutomationWindow.ExcludedTestsFilter", new CORE_IMAGE_BRUSH_SVG("Starship/Common/reject", Icon16x16));
+		Set("AutomationWindow.EditTags", new CORE_IMAGE_BRUSH_SVG("Starship/Common/TagSimple_16", Icon16x16));
+		Set("AutomationWindow.ApplyTagFilter", new CORE_IMAGE_BRUSH_SVG("Starship/Common/filter", Icon16x16));
+		Set("AutomationWindow.TagFilterIsActive", new CORE_IMAGE_BRUSH_SVG("Starship/Common/badge-modified", Icon16x16));
+		Set("AutomationWindow.RemoveTagFilter", new CORE_IMAGE_BRUSH_SVG("Starship/Common/close", Icon16x16));
 		Set("AutomationWindow.ErrorFilter.Small", new CORE_IMAGE_BRUSH_SVG("Starship/Common/x-circle", Icon16x16));
 		Set("AutomationWindow.WarningFilter.Small", new CORE_IMAGE_BRUSH_SVG("Starship/Common/alert-triangle", Icon16x16));
 		Set("AutomationWindow.SmokeTestFilter.Small", new CORE_IMAGE_BRUSH_SVG("Starship/Common/Test", Icon16x16));
diff --git a/Engine/Source/Developer/AutomationWindow/Private/SAutomationTestItem.cpp b/Engine/Source/Developer/AutomationWindow/Private/SAutomationTestItem.cpp
--- a/Engine/Source/Developer/AutomationWindow/Private/SAutomationTestItem.cpp
+++ b/Engine/Source/Developer/AutomationWindow/Private/SAutomationTestItem.cpp
@@ -11,8 +11,8 @@
 #include "Widgets/Images/SImage.h"
 #include "Widgets/Images/SSpinningImage.h"
 #include "Widgets/Notifications/SProgressBar.h"
-#include "Widgets/Text/STextBlock.h"
 #include "Widgets/Input/SCheckBox.h"
+#include "Widgets/Input/SSuggestionTextBox.h"
 #include "SSimpleComboButton.h"
 #include "AutomationWindowStyle.h"
 #include "AutomationTestExcludelist.h"
@@ -32,6 +32,8 @@
 #include "Widgets/Input/SHyperlink.h"
 #include "SSimpleButton.h"
 
+#include "SAutomationTestTagEditor.h"
+
 #define LOCTEXT_NAMESPACE "AutomationTestItem"
 
 namespace
@@ -189,7 +191,7 @@ TSharedRef<SWidget> SAutomationTestItem::GenerateWidgetForColumn( const FName& C
 	else if (ColumnName == AutomationTestWindowConstants::SkippedOptions)
 	{
 		return SNew(SHorizontalBox)
-			+SHorizontalBox::Slot()
+			+ SHorizontalBox::Slot()
 #if WITH_EDITOR
 			.HAlign(HAlign_Center)
 			.AutoWidth()
@@ -204,7 +206,33 @@ TSharedRef<SWidget> SAutomationTestItem::GenerateWidgetForColumn( const FName& C
 #endif
 			;
 	}
-	else if( ColumnName == AutomationTestWindowConstants::Title)
+	else if (ColumnName == AutomationTestWindowConstants::Tags)
+	{
+		return SNew(SHorizontalBox)
+			+ SHorizontalBox::Slot()
+			.HAlign(HAlign_Left)
+			.AutoWidth()
+			[
+				SNew(SSimpleButton)
+					.Icon(FAutomationWindowStyle::Get().GetBrush("Icons.Edit"))
+					.Visibility(this, &SAutomationTestItem::CanEditTags_GetVisibility)
+					.ToolTipText(LOCTEXT("EditTestTags", "Edit Tags"))
+					.OnClicked(FOnClicked::CreateSP(this, &SAutomationTestItem::OnEditTestTagsClicked))
+			]
+			+ SHorizontalBox::Slot()
+			.Expose(TagsSlot)
+			.HAlign(HAlign_Left)
+			.VAlign(VAlign_Center)
+			.AutoWidth()
+			[
+				SNew(STextBlock)
+					.Margin(FMargin(2, 0, 0, 0))
+					.OverflowPolicy(ETextOverflowPolicy::Ellipsis)
+					.Text(FText::FromString(TestStatus->GetTags()))
+					.ToolTipText(FText::FromString(TestStatus->GetTags()))
+			];
+	}
+	else if (ColumnName == AutomationTestWindowConstants::Title)
 	{
 		TSharedRef<SWidget> TestNameWidget = SNullWidget::NullWidget;
 
@@ -564,6 +592,11 @@ bool SAutomationTestItem::CanSkipFlagBeChanged() const
 	return WITH_EDITOR && IsLocalSession && !TestStatus->IsToBeSkippedByPropagation();
 }
 
+EVisibility SAutomationTestItem::CanEditTags_GetVisibility() const
+{
+	return (WITH_EDITOR && !TestStatus->IsParent()) ? EVisibility::Visible : EVisibility::Collapsed;
+}
+
 #if WITH_EDITOR
 void PopulateMenuContent(FMenuBuilder* MenuBuilder, TSet<FName>* OptionsDestination, const TSet<FName>& ItemNames)
 {
@@ -864,6 +897,26 @@ FReply SAutomationTestItem::OnEditExcludeOptionsClicked()
 	return FReply::Handled();
 }
 
+FReply SAutomationTestItem::OnEditTestTagsClicked()
+{
+#if WITH_EDITOR
+	TArray<IAutomationReportPtr> ReportList{ TestStatus };
+	SAutomationTestTagEditor::ShowInWindow(this->AsShared(), ReportList,
+		FOnTestTagsSavedEvent::CreateLambda([this]()
+		{
+			FChildren* ContainerWidgetChildren = TagsSlot->GetContainerWidget().GetChildren();
+			if (ContainerWidgetChildren && ContainerWidgetChildren->Num() > 1)
+			{
+				STextBlock& TagsText = static_cast<STextBlock&>(ContainerWidgetChildren->GetChildAt(1).Get());
+				FString UpdatedTags = TestStatus->GetTags();
+				TagsText.SetText(FText::FromString(UpdatedTags));
+				TagsText.SetToolTipText(FText::FromString(UpdatedTags));
+			}
+		}));
+#endif
+	return FReply::Handled();
+}
+
 FSlateColor SAutomationTestItem::ItemStatus_BackgroundColor(const int32 ClusterIndex) const
 {
 	if (TestStatus->GetTotalNumChildren()==0)
diff --git a/Engine/Source/Developer/AutomationWindow/Private/SAutomationTestItem.h b/Engine/Source/Developer/AutomationWindow/Private/SAutomationTestItem.h
--- a/Engine/Source/Developer/AutomationWindow/Private/SAutomationTestItem.h
+++ b/Engine/Source/Developer/AutomationWindow/Private/SAutomationTestItem.h
@@ -11,12 +11,13 @@
 #include "Widgets/Views/SListView.h"
 #include "Widgets/Views/STableViewBase.h"
 #include "Widgets/Views/STableRow.h"
+#include "Widgets/Text/STextBlock.h"
+
 
 class IAutomationReport;
 
 DECLARE_DELEGATE_OneParam(FOnItemCheckedStateChanged, TSharedPtr<IAutomationReport> /*TestStatus*/);
 
-
 /**
  * Implements a row widget for the automation list.
  */
@@ -136,6 +137,18 @@ protected:
 	*/
 	FReply OnEditExcludeOptionsClicked();
 
+	/**
+	 * Returns whether of not test tags can be edited on test.
+	*/
+	EVisibility CanEditTags_GetVisibility() const;
+
+	/**
+	* Open Test Tags editor.
+	* 
+	* @return Editor reply result.
+	*/
+	FReply OnEditTestTagsClicked();
+
 	/**
 	 * Returns color that indicates test status per cluster.
 	 *
@@ -232,4 +245,6 @@ private:
 
 	/** Holds a delegate to be invoked when the check box state changed. */
 	FOnItemCheckedStateChanged OnCheckedStateChangedDelegate;
+
+	SHorizontalBox::FSlot* TagsSlot;
 };
diff --git a/Engine/Source/Developer/AutomationWindow/Private/SAutomationTestTagEditor.h b/Engine/Source/Developer/AutomationWindow/Private/SAutomationTestTagEditor.h
new file mode 100644
--- /dev/null
+++ b/Engine/Source/Developer/AutomationWindow/Private/SAutomationTestTagEditor.h
@@ -0,0 +1,372 @@
+// Copyright Epic Games, Inc. All Rights Reserved.
+
+#pragma once
+
+#include "CoreMinimal.h"
+#include "Misc/CoreDelegates.h"
+#include "Misc/MessageDialog.h"
+#include "Widgets/DeclarativeSyntaxSupport.h"
+#include "Widgets/SWidget.h"
+#include "Widgets/SCompoundWidget.h"
+#include "Widgets/Input/SSuggestionTextBox.h"
+#include "Widgets/Layout/SBorder.h"
+#include "Widgets/Layout/SScrollBox.h"
+#include "Widgets/Views/SListView.h"
+#include "Widgets/Views/STableViewBase.h"
+#include "Widgets/Views/STreeView.h"
+#include "Widgets/Views/SHeaderRow.h"
+#include "Widgets/Views/STableRow.h"
+#include "Modules/ModuleManager.h"
+#include "Textures/SlateIcon.h"
+#include "Framework/MultiBox/MultiBoxBuilder.h"
+#include "Styling/AppStyle.h"
+#include "IAutomationControllerModule.h"
+
+#include "SAutomationTestTagListBox.h"
+
+#define LOCTEXT_NAMESPACE "SAutomationTestTagEditor"
+
+namespace TestNameAndTagTableConstants
+{
+	const FName Name(TEXT("Name"));
+	const FName Tags(TEXT("Tags"));
+};
+
+class STestNameAndTagTableRow : public SMultiColumnTableRow<TSharedPtr<FString>>
+{
+public:
+	SLATE_BEGIN_ARGS(STestNameAndTagTableRow) {}
+		SLATE_ARGUMENT(TSharedPtr<IAutomationReport>, TestReport)
+	SLATE_END_ARGS()
+
+	void Construct(const FArguments& InArgs, const TSharedRef<STableViewBase>& InOwnerTableView)
+	{
+		TestReport = InArgs._TestReport;
+
+		SMultiColumnTableRow<TSharedPtr<FString>>::Construct(FSuperRowType::FArguments(), InOwnerTableView);
+	}
+
+	BEGIN_SLATE_FUNCTION_BUILD_OPTIMIZATION
+	virtual TSharedRef<SWidget> GenerateWidgetForColumn(const FName& ColumnName) override
+	{
+		if (ColumnName == TestNameAndTagTableConstants::Name)
+		{
+			if (!TestReport->GetSourceFile().IsEmpty())
+			{
+				return SNew(SHyperlink)
+					.Style(FAutomationWindowStyle::Get(), "Common.GotoNativeCodeHyperlink")
+					.OnNavigate_Lambda([this]
+						{
+							FSlateApplication::Get()
+								.GotoLineInSource(TestReport->GetSourceFile(), TestReport->GetSourceFileLine());
+						})
+					.Text(FText::FromString(TestReport->GetFullTestPath()));
+			}
+			return SNew(STextBlock)
+				.Text(FText::FromString(TestReport->GetFullTestPath()))
+				.ToolTipText(FText::FromString(TestReport->GetFullTestPath()));
+		}
+		else if (ColumnName == TestNameAndTagTableConstants::Tags)
+		{
+			return SNew(STextBlock)
+				.Text(FText::FromString(TestReport->GetTags()))
+				.ToolTipText(FText::FromString(TestReport->GetTags()));
+		}
+		return SNullWidget::NullWidget;
+	}
+	END_SLATE_FUNCTION_BUILD_OPTIMIZATION
+
+private:
+
+	IAutomationReportPtr TestReport;
+};
+
+DECLARE_DELEGATE(FOnTestTagsSavedEvent);
+
+class SAutomationTestTagEditor
+	: public SCompoundWidget
+{
+public:
+
+	SLATE_BEGIN_ARGS(SAutomationTestTagEditor) {}
+	SLATE_END_ARGS()
+
+public:
+
+	/**
+	 * Construct this widget.
+	 *
+	 * @param InArgs The declaration data for this widget.
+	 */
+	void Construct(const FArguments& InArgs, const TArray<IAutomationReportPtr>& InTestReports)
+	{
+		TestReports = InTestReports;
+
+		TArray<FString> ExistingTags;
+		TSet<FString> ImmutableTags;
+		if (InTestReports.Num() == 1)
+		{
+			IAutomationReportPtr Report = InTestReports[0];
+			ExistingTags = Report->GetTagsArray();
+
+			for (FString TagEx : ExistingTags)
+			{
+				if (FAutomationTestFramework::Get().IsTagImmutable(Report->GetFullTestPath(), TagEx))
+				{
+					ImmutableTags.Add(TagEx);
+				}
+			}
+		}
+
+		ChildSlot
+			[
+				SNew(SBox)
+					.WidthOverride(640)
+					[
+						SNew(SVerticalBox)
+							+ SVerticalBox::Slot()
+							.MaxHeight(200.f)
+							.Padding(5.0f)
+							[
+								SNew(SListView<TSharedPtr<IAutomationReport>>)
+									.SelectionMode(ESelectionMode::None)
+									.ListItemsSource(&InTestReports)
+									.OnGenerateRow(this, &SAutomationTestTagEditor::OnGenerateTestNameAndTagRow)
+									.HeaderRow
+									(
+										SNew(SHeaderRow)
+										+ SHeaderRow::Column(TestNameAndTagTableConstants::Name)
+										.FillWidth(0.7f)
+										.HAlignHeader(HAlign_Center)
+										.VAlignHeader(VAlign_Center)
+										.HAlignCell(HAlign_Left)
+										.VAlignCell(VAlign_Center)
+										[
+											SNew(SHorizontalBox)
+												+ SHorizontalBox::Slot()
+												.AutoWidth()
+												.VAlign(VAlign_Center)
+												[
+													SNew(STextBlock)
+														.Text(LOCTEXT("TestName_Header", "Test"))
+												]
+										]
+
+										+ SHeaderRow::Column(TestNameAndTagTableConstants::Tags)
+										.FillWidth(0.3f)
+										.HAlignHeader(HAlign_Center)
+										.VAlignHeader(VAlign_Center)
+										.HAlignCell(HAlign_Left)
+										.VAlignCell(VAlign_Center)
+										[
+											SNew(SHorizontalBox)
+												+ SHorizontalBox::Slot()
+												.AutoWidth()
+												.VAlign(VAlign_Center)
+												[
+													SNew(STextBlock)
+														.Text(LOCTEXT("TestTags_Header", "Tags"))
+												]
+										]
+									)
+							]
+							+ SVerticalBox::Slot()
+							.AutoHeight()
+							.Padding(5.0f)
+							[
+								SAssignNew(TagSuggestionBox, SSuggestionTextBox)
+									.ForegroundColor(FSlateColor::UseForeground())
+									.CommitOnSuggestionSelected(true)
+									.ClearKeyboardFocusOnCommit(false)
+									.OnTextCommitted(this, &SAutomationTestTagEditor::TagSuggestionBox_SetText)
+									.HintText(LOCTEXT("EnterTagHint", "Enter tag"))
+									.OnShowingSuggestions(this, &SAutomationTestTagEditor::TagSuggestionBox_GetSuggestions)
+							]
+							+ SVerticalBox::Slot()
+							.AutoHeight()
+							.Padding(5.0f)
+							[
+								SAssignNew(TagListBox, SAutomationTestTagListBox, ExistingTags, ImmutableTags)
+							]
+							+ SVerticalBox::Slot()
+							.AutoHeight()
+							.Padding(5.0f)
+							.HAlign(HAlign_Center)
+							[
+								SNew(SHorizontalBox)
+									+ SHorizontalBox::Slot()
+									.AutoWidth()
+									.HAlign(HAlign_Center)
+									.VAlign(VAlign_Center)
+									.Padding(5, 0)
+									[
+										SNew(SButton)
+											.Text(LOCTEXT("SaveTags", "Save"))
+											.IsEnabled(this, &SAutomationTestTagEditor::SaveButonEnabled)
+											.OnClicked(this, &SAutomationTestTagEditor::OnSaveButtonClicked)
+									]
+									+ SHorizontalBox::Slot()
+									.AutoWidth()
+									.HAlign(HAlign_Center)
+									.VAlign(VAlign_Center)
+									.Padding(5, 0)
+									[
+										SNew(SButton)
+											.Text(LOCTEXT("CancelTags", "Discard"))
+											.OnClicked(this, &SAutomationTestTagEditor::OnDiscardButtonClicked)
+									]
+							]
+					]
+			];
+	}
+
+	static void ShowInWindow(const TSharedRef<const SWidget> InParentWindow, const TArray<TSharedPtr<IAutomationReport>>& InTestReports, const FOnTestTagsSavedEvent& OnTestTagsSaved)
+	{
+		TSharedPtr<SAutomationTestTagEditor> TestTagEditor;
+		TSharedRef<SWidget> TagsForm = SNew(SBox)
+			.WidthOverride(640)
+			[
+				SAssignNew(TestTagEditor, SAutomationTestTagEditor, InTestReports)
+			];
+
+		TSharedRef<SWindow> Window = SNew(SWindow)
+			.Title(LOCTEXT("TestTagsEditor", "Edit Tags"))
+			.SizingRule(ESizingRule::Autosized)
+			.ClientSize(FVector2D(480, 160))
+			.SupportsMinimize(false)
+			.SupportsMaximize(false)
+			.AutoCenter(EAutoCenter::PrimaryWorkArea);
+
+		Window->SetContent
+		(
+			TagsForm
+		);
+
+		TestTagEditor->SetOnTestTagsSavedEvent(OnTestTagsSaved);
+
+		FSlateApplication::Get().AddModalWindow(Window, InParentWindow);
+	}
+
+	void SetOnTestTagsSavedEvent(const FOnTestTagsSavedEvent& Event)
+	{
+		OnTestTagsSavedEvent = Event;
+	}
+
+private:
+	void TagSuggestionBox_SetText(const FText& NewText, ETextCommit::Type CommitInfo)
+	{
+		if (CommitInfo == ETextCommit::Type::OnEnter && !NewText.ToString().IsEmpty())
+		{
+			TagListBox->AddTag(NewText.ToString(), true, false);
+			TagSuggestionBox->SetText(FText::GetEmpty());
+		}
+	}
+
+	void TagSuggestionBox_GetSuggestions(const FString& CurrText, TArray<FString>& OutSuggestions)
+	{
+		OutSuggestions = FAutomationTestFramework::Get().GetAllExistingTags().Array().FilterByPredicate([&CurrText](const FString& CurrentTag) { return CurrentTag.Contains(CurrText); });
+	}
+
+	FReply OnDiscardButtonClicked()
+	{
+		TSharedPtr<SWindow> Window = FSlateApplication::Get().FindWidgetWindow(this->AsShared());
+
+		TArray<FString> NewTags = TagListBox->GetNewTags();
+		TArray<FString> MarkedForDeleteTags = TagListBox->GetMarkForDeleteTags();
+		if (NewTags.Num() > 0 || MarkedForDeleteTags.Num() > 0)
+		{
+			if (EAppReturnType::Yes == FMessageDialog::Open(EAppMsgType::YesNo, FText::FromString(TEXT("Discard changes?"))))
+			{
+				FSlateApplication::Get().DestroyWindowImmediately(Window.ToSharedRef());
+			}
+		}
+		else
+		{
+			FSlateApplication::Get().DestroyWindowImmediately(Window.ToSharedRef());
+		}
+		return FReply::Handled();
+	}
+
+	bool SaveButonEnabled() const
+	{
+		return TagListBox->GetNewTags().Num() > 0 || TagListBox->GetMarkForDeleteTags().Num() > 0;
+	}
+
+
+	FReply OnSaveButtonClicked()
+	{
+		TSharedPtr<SWindow> Window = FSlateApplication::Get().FindWidgetWindow(this->AsShared());
+
+		TArray<FString> NewTags = TagListBox->GetNewTags();
+		TArray<FString> MarkedForDeleteTags = TagListBox->GetMarkForDeleteTags();
+		
+		if (NewTags.Num() > 0 || MarkedForDeleteTags.Num() > 0)
+		{
+			FString Changes;
+			if (NewTags.Num() > 0)
+			{
+				Changes = Changes.Append(TEXT("New tags: ")).Append(FString::Join(NewTags, TEXT(","))).Append(TEXT(".\n"));
+			}
+			if (MarkedForDeleteTags.Num() > 0)
+			{
+				Changes = Changes.Append(TEXT("Removing tags: ")).Append(FString::Join(MarkedForDeleteTags, TEXT(","))).Append(TEXT("."));
+			}
+			FText ChangesText = FText::FromString(*Changes);
+			if (EAppReturnType::Yes == FMessageDialog::Open(EAppMsgType::YesNo, FText::Format(LOCTEXT("TagsChangedSaveMsg", "Confirm tag changes to apply.\n{0}"), ChangesText)))
+			{
+				TArray<FString> Tags;
+				TArray<FString> TestNames;
+				TArray<FString> TestFilePaths;
+				for (IAutomationReportPtr Report : TestReports)
+				{
+					TestNames.Add(Report->GetFullTestPath());
+					TestFilePaths.Add(Report->GetSourceFile());
+					Tags = Report->GetTagsArray();
+					for (FString DeleteTag : MarkedForDeleteTags)
+					{
+						NewTags.Remove(DeleteTag);
+						if (FAutomationTestFramework::Get().IsTagImmutable(Report->GetFullTestPath(), DeleteTag) && Tags.Contains(DeleteTag))
+						{
+							// Do not delete own immutable tag
+							continue;
+						}
+						Tags.Remove(DeleteTag);
+					}
+					Tags.Append(NewTags);
+					bool Registered = FAutomationTestFramework::Get().RegisterAutomationTestTags(
+						Report->GetFullTestPath(),
+						FString::JoinBy(Tags, TEXT(""),
+							[](const FString& InTag)
+							{
+								return FString::Printf(TEXT("[%s]"), *InTag);
+							}), false);
+					check(Registered);
+				}
+				FAutomationTestFramework::Get().SaveTestTagMappings(TestNames, TestFilePaths);
+
+				OnTestTagsSavedEvent.ExecuteIfBound();
+
+				FSlateApplication::Get().DestroyWindowImmediately(Window.ToSharedRef());
+			}
+		}
+		return FReply::Handled();
+	}
+
+	TSharedRef<ITableRow> OnGenerateTestNameAndTagRow(TSharedPtr<IAutomationReport> InTestReport, const TSharedRef<STableViewBase>& OwnerTable)
+	{
+		return SNew(STestNameAndTagTableRow, OwnerTable).TestReport(InTestReport);
+	}
+
+	void OnGetTestNameAndTagChildren(TSharedPtr<IAutomationReport> InItem, TArray<TSharedPtr<IAutomationReport>>& OutItems)
+	{
+		OutItems = InItem.Get()->GetChildReports();
+	}
+
+	TSharedPtr<SSuggestionTextBox> TagSuggestionBox;
+	TSharedPtr<SAutomationTestTagListBox> TagListBox;
+	TArray<IAutomationReportPtr> TestReports;
+	FOnTestTagsSavedEvent OnTestTagsSavedEvent;
+};
+
+
+#undef LOCTEXT_NAMESPACE
\ No newline at end of file
diff --git a/Engine/Source/Developer/AutomationWindow/Private/SAutomationTestTagFilter.h b/Engine/Source/Developer/AutomationWindow/Private/SAutomationTestTagFilter.h
new file mode 100644
--- /dev/null
+++ b/Engine/Source/Developer/AutomationWindow/Private/SAutomationTestTagFilter.h
@@ -0,0 +1,186 @@
+// Copyright Epic Games, Inc. All Rights Reserved.
+
+#pragma once
+
+#include "CoreMinimal.h"
+#include "Misc/CoreDelegates.h"
+#include "Misc/MessageDialog.h"
+#include "Widgets/DeclarativeSyntaxSupport.h"
+#include "Widgets/SWidget.h"
+#include "Widgets/SCompoundWidget.h"
+#include "Widgets/Layout/SBorder.h"
+#include "Widgets/Layout/SScrollBox.h"
+#include "Widgets/Layout/SGridPanel.h"
+#include "Widgets/Images/SLayeredImage.h"
+#include "Modules/ModuleManager.h"
+#include "Textures/SlateIcon.h"
+#include "Framework/MultiBox/MultiBoxBuilder.h"
+#include "Styling/AppStyle.h"
+#include "IAutomationControllerModule.h"
+
+#define LOCTEXT_NAMESPACE "SAutomationTestTagFilter"
+
+DECLARE_DELEGATE(FOnTestFilterAppliedEvent);
+
+class SAutomationTestTagFilter
+	: public SCompoundWidget
+{
+public:
+
+	SLATE_BEGIN_ARGS(SAutomationTestTagFilter) {}
+	SLATE_END_ARGS()
+
+public:
+
+	/**
+	 * Construct this widget.
+	 *
+	 * @param InArgs The declaration data for this widget.
+	 */
+	void Construct(const FArguments& InArgs)
+	{
+		ChildSlot
+			[
+				SNew(SBox)
+					.WidthOverride(300)
+					[
+						SNew(SHorizontalBox)
+						+ SHorizontalBox::Slot()
+						.MinWidth(150.f)
+						.MaxWidth(250.f)
+						.AutoWidth()
+						.Padding(5.0f)
+						[
+							SAssignNew(FilterText, SEditableTextBox)
+								.ForegroundColor(FSlateColor::UseForeground())
+								.ClearKeyboardFocusOnCommit(false)
+								.HintText(LOCTEXT("EnterFilterHint", "Enter tag filter"))
+						]
+						+ SHorizontalBox::Slot()
+						.AutoWidth()
+						.HAlign(HAlign_Center)
+						.VAlign(VAlign_Center)
+						.Padding(5, 0)
+						[
+							SNew(SButton)
+								.ToolTipText(LOCTEXT("ApplyTagFilterHint", "Apply Filter"))
+								.IsEnabled_Lambda([this]()
+									{
+										return !FilterText->GetText().IsEmptyOrWhitespace();
+									})
+								.OnClicked(this, &SAutomationTestTagFilter::OnApplyTagFilterButtonClicked)
+								[
+									SAssignNew(FilterButtonImage, SLayeredImage)
+										.Image(FAutomationWindowStyle::Get().GetBrush("AutomationWindow.ApplyTagFilter"))
+								]
+						]
+						+ SHorizontalBox::Slot()
+						.AutoWidth()
+						.HAlign(HAlign_Center)
+						.VAlign(VAlign_Center)
+						.Padding(5, 0)
+						[
+							SNew(SButton)
+								.ToolTipText(LOCTEXT("RemoveTagFilterHint", "Remove Filter"))
+								.IsEnabled_Lambda([this]()
+									{
+										return IsFilterApplied;
+									})
+								.OnClicked(this, &SAutomationTestTagFilter::OnRemoveTagFilterButtonClicked)
+								[
+									SNew(SImage)
+										.Image(FAutomationWindowStyle::Get().GetBrush("AutomationWindow.RemoveTagFilter"))
+								]
+						]
+					]
+			];
+
+		FilterButtonImage->AddLayer(TAttribute<const FSlateBrush*>::CreateLambda([this]() -> const FSlateBrush*
+			{
+				if (!IsFilterApplied)
+				{
+					return nullptr;
+				}
+
+				return FAutomationWindowStyle::Get().GetBrush("AutomationWindow.TagFilterIsActive");
+			}));
+	}
+
+	FReply OnApplyTagFilterButtonClicked()
+	{
+		FString FilterTextString = FilterText->GetText().ToString();
+		if (!FilterTextString.IsEmpty())
+		{
+			ApplyTagFilter(FilterText->GetText().ToString());
+			OnTestFilterApplied.ExecuteIfBound();
+			IsFilterApplied = true;
+		}
+		return FReply::Handled();
+	}
+
+	FReply OnRemoveTagFilterButtonClicked()
+	{
+		ClearFilter();
+		return FReply::Handled();
+	}
+
+	void ApplyTagFilter(const FString& TagFilter)
+	{
+		TSharedPtr <AutomationFilterCollection> AutomationFilters = MakeShareable(new AutomationFilterCollection());
+		FAutomationGroupFilter* FilterTags = new FAutomationGroupFilter();
+		TArray<FAutomatedTestTagFilter> TagList;
+		TagList.Add(FAutomatedTestTagFilter(TagFilter));
+		FilterTags->SetTagFilter(TagList);
+		AutomationFilters->Add(MakeShareable(FilterTags));
+
+		FModuleManager::GetModuleChecked<IAutomationControllerModule>("AutomationController").GetAutomationController()->SetFilter(AutomationFilters);
+	}
+
+	static void ShowInWindow(const TSharedRef<const SWidget> InParentWindow, const FOnTestFilterAppliedEvent& OnTestFilterAppliedEvent)
+	{
+		TSharedPtr<SAutomationTestTagFilter> TestTagFilter;
+		TSharedRef<SWidget> TagsFilterForm = SNew(SBox)
+			.WidthOverride(640)
+			[
+				SAssignNew(TestTagFilter, SAutomationTestTagFilter)
+			];
+
+		TSharedRef<SWindow> Window = SNew(SWindow)
+			.Title(LOCTEXT("TestTagsFilter", "Filter Tags"))
+			.SizingRule(ESizingRule::Autosized)
+			.SupportsMinimize(false)
+			.SupportsMaximize(false)
+			.AutoCenter(EAutoCenter::PrimaryWorkArea);
+
+		Window->SetContent
+		(
+			TagsFilterForm
+		);
+
+		TestTagFilter->SetOnTestFilterAppliedEvent(OnTestFilterAppliedEvent);
+
+		FSlateApplication::Get().AddModalWindow(Window, InParentWindow);
+	}
+
+	void SetOnTestFilterAppliedEvent(const FOnTestFilterAppliedEvent& OnTestFilterAppliedEvent)
+	{
+		OnTestFilterApplied = OnTestFilterAppliedEvent;
+	}
+
+	void ClearFilter()
+	{
+		ApplyTagFilter(TEXT(""));
+		OnTestFilterApplied.ExecuteIfBound();
+		FilterText->SetText(FText::GetEmpty());
+		IsFilterApplied = false;
+	}
+
+protected:
+	FOnTestFilterAppliedEvent OnTestFilterApplied;
+	TSharedPtr<SEditableTextBox> FilterText;
+	TSharedPtr<SLayeredImage> FilterButtonImage;
+	bool IsFilterApplied;
+};
+
+
+#undef LOCTEXT_NAMESPACE
\ No newline at end of file
diff --git a/Engine/Source/Developer/AutomationWindow/Private/SAutomationTestTagListBox.h b/Engine/Source/Developer/AutomationWindow/Private/SAutomationTestTagListBox.h
new file mode 100644
--- /dev/null
+++ b/Engine/Source/Developer/AutomationWindow/Private/SAutomationTestTagListBox.h
@@ -0,0 +1,211 @@
+// Copyright Epic Games, Inc. All Rights Reserved.
+
+#pragma once
+
+#include "CoreMinimal.h"
+#include "Widgets/DeclarativeSyntaxSupport.h"
+#include "Widgets/SWidget.h"
+#include "Widgets/SCompoundWidget.h"
+#include "Widgets/Layout/SBorder.h"
+#include "Widgets/Layout/SWrapBox.h"
+#include "Modules/ModuleManager.h"
+#include "Textures/SlateIcon.h"
+#include "Framework/MultiBox/MultiBoxBuilder.h"
+#include "Styling/AppStyle.h"
+#include "IAutomationControllerModule.h"
+
+#define LOCTEXT_NAMESPACE "SAutomationTestTagListBox"
+
+class SAutomationTestTagItem
+	: public SCompoundWidget
+{
+public:
+
+	SLATE_BEGIN_ARGS(SAutomationTestTagItem) {}
+	SLATE_END_ARGS()
+
+public:
+
+	/**
+	 * Construct this widget.
+	 *
+	 * @param InArgs The declaration data for this widget.
+	 */
+	void Construct(const FArguments& InArgs, const FString& InTag, const bool InIsNew, const bool InImmutable)
+	{
+		TagValue = InTag;
+		IsNew = InIsNew;
+		FText TagText = FText::Format(LOCTEXT("NewTestTag", "{0}"), FText::FromString(InTag));
+		FText TagToolTipText = FText::Format(LOCTEXT("NewTestTagToolTip", "{0}"),
+			FText::FromString(InTag + (InImmutable ? TEXT(" (immutable, from source code)") : TEXT(""))));
+
+		ChildSlot
+			[
+				SAssignNew(MainBorder, SBorder)
+					.IsEnabled(!InImmutable)
+					.BorderImage(FAppStyle::GetBrush("ToolPanel.GroupBorder"))
+					.BorderBackgroundColor(FColorList::Black)
+					.Content()
+					[
+						SNew(SHorizontalBox)
+							.Clipping(EWidgetClipping::ClipToBoundsWithoutIntersecting)
+							+ SHorizontalBox::Slot()
+							.AutoWidth()
+							.VAlign(EVerticalAlignment::VAlign_Center)
+							.HAlign(EHorizontalAlignment::HAlign_Left)
+							.Padding(4, 2, 2, 2)
+							[
+								SNew(STextBlock)
+									.ColorAndOpacity(FLinearColor(.8f, .8f, .8f, 1.0f))
+									.Font(FCoreStyle::GetDefaultFontStyle("Regular", 10))
+									.Text(TagText)
+									.ToolTipText(TagToolTipText)
+							]
+							+ SHorizontalBox::Slot()
+							.Padding(2)
+							[
+								SNew(SButton)
+									.Visibility(!InImmutable ? EVisibility::Visible : EVisibility::Collapsed)
+									.ButtonStyle(FAppStyle::Get(), TEXT("SimpleButton"))
+									.OnClicked(this, &SAutomationTestTagItem::ToggleTagDeletion)
+									.VAlign(VAlign_Center)
+									.HAlign(HAlign_Center)
+									[
+										SNew(SImage)
+											.Image(FAppStyle::GetBrush("Icons.Delete"))
+											.ColorAndOpacity(FSlateColor::UseForeground())
+									]
+							]
+					]
+			];
+	}
+
+	bool IsMarkedForDeletion()
+	{
+		return MarkedForDeletion;
+	}
+
+	FString GetTagValue()
+	{
+		return TagValue;
+	}
+
+protected:
+	FReply ToggleTagDeletion()
+	{
+		MarkedForDeletion = !MarkedForDeletion;
+		if (!IsNew)
+		{
+			MainBorder->SetBorderBackgroundColor(MarkedForDeletion ? FColorList::Red : FColorList::Black);
+		}
+		else
+		{
+			ChildSlot.DetachWidget();
+		}
+		return FReply::Handled();
+	}
+
+	TSharedPtr<SBorder> MainBorder;
+	bool IsNew;
+	bool MarkedForDeletion;
+	FString TagValue;
+};
+
+class SAutomationTestTagListBox
+	: public SCompoundWidget
+{
+public:
+
+	SLATE_BEGIN_ARGS(SAutomationTestTagListBox) {}
+	SLATE_END_ARGS()
+
+public:
+
+	/**
+	 * Construct this widget.
+	 *
+	 * @param InArgs The declaration data for this widget.
+	 * @param InCurrentTestStatus The current test report pointer.
+	 */
+	void Construct( const FArguments& InArgs, TArray<FString>& InExistingTags, TSet<FString> InImmutableTags )
+	{
+		ExistingTags = InExistingTags;
+
+		ChildSlot
+		[
+			SNew(SBorder)
+				.BorderImage(FAppStyle::GetBrush("ToolPanel.GroupBorder"))
+				.BorderBackgroundColor(FLinearColor::Transparent)
+				.Content()
+				[
+					SAssignNew(TagBox, SWrapBox)
+						.UseAllottedSize(true)
+						.Clipping(EWidgetClipping::ClipToBoundsWithoutIntersecting)
+				]
+		];
+
+		for (const FString& TestTag : ExistingTags)
+		{
+			AddTag(TestTag, false, InImmutableTags.Contains(TestTag));
+		}
+	}
+
+	void AddTag(const FString& TestTag, bool IsNew, bool IsImmutable)
+	{
+		bool AlreadyInNewTagsCaseInsensitive =
+			NewTags.FindByPredicate([TestTag](const FString& Item)
+				{
+					return Item.Compare(TestTag, ESearchCase::IgnoreCase) == 0;
+				}) != nullptr;
+		bool AlreadyInExistingTagsCaseInsensitive =
+			ExistingTags.FindByPredicate([TestTag](const FString& Item)
+				{
+					return Item.Compare(TestTag, ESearchCase::IgnoreCase) == 0;
+				}) != nullptr;
+
+		if (AlreadyInNewTagsCaseInsensitive || (IsNew && AlreadyInExistingTagsCaseInsensitive))
+			return;
+
+		if (IsNew)
+		{
+			NewTags.Add(TestTag);
+		}
+		TagBox->AddSlot()
+			.FillEmptySpace(false)
+			.VAlign(EVerticalAlignment::VAlign_Center)
+			.HAlign(EHorizontalAlignment::HAlign_Left)
+			.Padding(FMargin(2.0f))
+			[
+				SNew(SAutomationTestTagItem, TestTag, IsNew, IsImmutable)
+			];
+	}
+
+	TArray<FString> GetNewTags()
+	{
+		return NewTags;
+	}
+
+	TArray<FString> GetMarkForDeleteTags()
+	{
+		TArray<FString> ToDelete;
+		TagBox->GetAllChildren()->ForEachWidget([this, &ToDelete](SWidget& TagItem)
+			{
+				SAutomationTestTagItem* TestTagItem = (SAutomationTestTagItem*)(&TagItem);
+				if (TestTagItem && TestTagItem->IsMarkedForDeletion())
+				{
+					ToDelete.Add(TestTagItem->GetTagValue());
+				}
+			});
+
+		return ToDelete;
+	}
+
+private:
+
+	TArray<FString> ExistingTags;
+	TArray<FString> NewTags;
+	TSharedPtr<SWrapBox> TagBox;
+};
+
+
+#undef LOCTEXT_NAMESPACE
\ No newline at end of file
diff --git a/Engine/Source/Developer/AutomationWindow/Private/SAutomationWindow.cpp b/Engine/Source/Developer/AutomationWindow/Private/SAutomationWindow.cpp
--- a/Engine/Source/Developer/AutomationWindow/Private/SAutomationWindow.cpp
+++ b/Engine/Source/Developer/AutomationWindow/Private/SAutomationWindow.cpp
@@ -32,6 +32,7 @@
 #include "SSimpleComboButton.h"
 
 #if WITH_EDITOR
+	#include "Dialogs/Dialogs.h"
 	#include "Engine/Level.h"
 	#include "Engine/World.h"
 	#include "FileHelpers.h"
@@ -47,6 +48,9 @@
 #include "AutomationControllerSettings.h"
 #include "Framework/Notifications/NotificationManager.h"
 
+#include "SAutomationTestTagEditor.h"
+#include "SAutomationTestTagFilter.h"
+
 
 #define LOCTEXT_NAMESPACE "AutomationTest"
 
@@ -74,6 +78,7 @@ public:
 		UI_COMMAND( WarningFilter, "Warnings", "Toggle Warning Filter", EUserInterfaceActionType::ToggleButton, FInputChord() );
 		UI_COMMAND( DeveloperDirectoryContent, "Dev Content", "Developer Directory Content Filter (when enabled, developer directories are also included)", EUserInterfaceActionType::ToggleButton, FInputChord() );
 		UI_COMMAND( ExcludedTestsFilter, "Excluded Tests", "Toggle Excluded Tests only", EUserInterfaceActionType::ToggleButton, FInputChord());
+		UI_COMMAND( EditTags, "Edit Tags for Selection", "Edit Tags for Selection", EUserInterfaceActionType::Button, FInputChord());
 
 #if WITH_EDITOR
 		// Added button for running the currently open level test.
@@ -91,6 +96,7 @@ public:
 	TSharedPtr<FUICommandInfo> WarningFilter;
 	TSharedPtr<FUICommandInfo> DeveloperDirectoryContent;
 	TSharedPtr<FUICommandInfo> ExcludedTestsFilter;
+	TSharedPtr<FUICommandInfo> EditTags;
 
 #if WITH_EDITOR
 	TSharedPtr<FUICommandInfo> RunLevelTest;
@@ -274,7 +280,18 @@ void SAutomationWindow::Construct( const FArguments& InArgs, const IAutomationCo
 					.Text( LOCTEXT("TestName_Header", "Test") )
 				]
 			]
-	
+
+			+ SHeaderRow::Column(AutomationTestWindowConstants::Tags)
+			.FillWidth(0.20f)
+			.HAlignHeader(HAlign_Left)
+			.VAlignHeader(VAlign_Center)
+			.HAlignCell(HAlign_Left)
+			.VAlignCell(VAlign_Center)
+			[
+				SNew(STextBlock)
+					.Text(LOCTEXT("TagsHeader", "Tags"))
+			]
+
 			+ SHeaderRow::Column( AutomationTestWindowConstants::SmokeTest )
 			.FixedWidth( 30.0f )
 			.HAlignHeader(HAlign_Center)
@@ -666,6 +683,8 @@ void SAutomationWindow::Construct( const FArguments& InArgs, const IAutomationCo
 	SessionManager->OnSessionInstanceUpdated().AddSP( this, &SAutomationWindow::HandleSessionManagerInstanceChanged );
 
 	FindWorkers();
+
+	FAutomationTestFramework::Get().LoadTestTagMappings();
 }
 END_SLATE_FUNCTION_BUILD_OPTIMIZATION
 
@@ -743,6 +762,11 @@ void SAutomationWindow::CreateCommands()
 		FIsActionChecked::CreateRaw( this, &SAutomationWindow::IsExcludedTestsFilterOn )
 		);
 
+	ActionList.MapAction( Commands.EditTags,
+		FExecuteAction::CreateRaw(this, &SAutomationWindow::HandleEditTagsButtonClicked ),
+		FCanExecuteAction::CreateRaw(this, &SAutomationWindow::IsEditTagsButtonEnabled )
+		);
+
 	// Added button for running the currently open level test.
 #if WITH_EDITOR
 	ActionList.MapAction(Commands.RunLevelTest,
@@ -777,7 +801,7 @@ TSharedRef< SWidget > SAutomationWindow::MakeAutomationWindowToolBar( const TSha
 {
 	struct Local
 	{
-		static void FillToolbar(FToolBarBuilder& ToolbarBuilder, TSharedRef<SWidget> RunTests, TSharedRef<SWidget> DeviceGroups, TSharedRef<SWidget> PresetBox, TWeakPtr<class SAutomationWindow> InAutomationWindow)
+		static void FillToolbar(FToolBarBuilder& ToolbarBuilder, TSharedRef<SWidget> RunTests, TSharedRef<SWidget> DeviceGroups, TSharedRef<SWidget> PresetBox, TSharedRef<SWidget> TagFilter, TWeakPtr<class SAutomationWindow> InAutomationWindow)
 		{
 			ToolbarBuilder.BeginSection("Automation");
 			{
@@ -814,6 +838,15 @@ TSharedRef< SWidget > SAutomationWindow::MakeAutomationWindowToolBar( const TSha
 				ToolbarBuilder.AddWidget(SNew(SSpacer));
 			}
 			ToolbarBuilder.EndSection();
+			ToolbarBuilder.BeginSection("Tags");
+			{
+				ToolbarBuilder.AddWidget(SNew(SSpacer));
+				ToolbarBuilder.AddToolBarButton(FAutomationWindowCommands::Get().EditTags);
+				ToolbarBuilder.AddWidget(SNew(SSpacer));
+				ToolbarBuilder.AddWidget( TagFilter );
+				ToolbarBuilder.AddWidget(SNew(SSpacer));
+			}
+			ToolbarBuilder.EndSection();
 		}
 	};
 
@@ -975,10 +1008,19 @@ TSharedRef< SWidget > SAutomationWindow::MakeAutomationWindowToolBar( const TSha
 			]
 		];
 
+	TagFilter = SNew(SAutomationTestTagFilter);
+	TagFilter->SetOnTestFilterAppliedEvent(FOnTestFilterAppliedEvent::CreateLambda([this]()
+		{
+			AutomationController->SetEnabledTests(TArray<FString>());
+			TestTable->RequestTreeRefresh();
+			TestTable->ClearExpandedItems();
+			TestTable->RebuildList();
+		}));
+
 	FToolBarBuilder ToolbarBuilder( InCommandList, FMultiBoxCustomization::None );
 	ToolbarBuilder.SetLabelVisibility(bIsLabelVisibilityEnabled ? EVisibility::Visible : EVisibility::Collapsed);
 	TWeakPtr<SAutomationWindow> AutomationWindow = SharedThis(this);
-	Local::FillToolbar(ToolbarBuilder, RunTests, DeviceGroups, TestPresets, AutomationWindow);
+	Local::FillToolbar(ToolbarBuilder, RunTests, DeviceGroups, TestPresets, TagFilter.ToSharedRef(), AutomationWindow);
 
 	// Create the tool bar!
 	return
@@ -2631,6 +2673,23 @@ bool SAutomationWindow::HandleMainContentIsEnabled() const
 	return (SessionManager->GetSelectedInstances().Num() > 0);
 }
 
+void SAutomationWindow::HandleEditTagsButtonClicked() const
+{
+#if WITH_EDITOR
+	SAutomationTestTagEditor::ShowInWindow(this->AsShared(), AutomationController->GetEnabledReports(),
+		FOnTestTagsSavedEvent::CreateLambda([this]()
+			{
+				AutomationController->OnTestsRefreshed().Broadcast();
+				TagFilter->ClearFilter();
+			}));
+#endif
+}
+
+bool SAutomationWindow::IsEditTagsButtonEnabled() const
+{
+	return AutomationController->GetEnabledReports().Num() > 0;
+}
+
 #if WITH_EDITOR
 // React to asset registry finishing updating.
 // We only want to do this if there are no tests already listed, otherwise this fires every time you save a map for example.
diff --git a/Engine/Source/Developer/AutomationWindow/Private/SAutomationWindow.h b/Engine/Source/Developer/AutomationWindow/Private/SAutomationWindow.h
--- a/Engine/Source/Developer/AutomationWindow/Private/SAutomationWindow.h
+++ b/Engine/Source/Developer/AutomationWindow/Private/SAutomationWindow.h
@@ -24,6 +24,7 @@
 class FAutomationFilter;
 class FAutomationGroupFilter;
 struct FAutomatedTestFilter;
+class SAutomationTestTagFilter;
 class FAutomationTestPresetManager;
 class FUICommandList;
 class SAutomationWindowCommandBar;
@@ -38,6 +39,7 @@ namespace AutomationTestWindowConstants
 	const FName Checked( TEXT("Checked") );
 	const FName Skipped( TEXT("Skipped") );
 	const FName SkippedOptions( TEXT("SkippedOptions") );
+	const FName Tags(TEXT("Tags"));
 	const FName Title( TEXT("Name") );
 	const FName SmokeTest( TEXT("SmokeTest") );
 	const FName RequiredDeviceCount( TEXT("RequiredDeviceCount") );
@@ -522,6 +524,11 @@ private:
 	/** Gets the combo box text to display for the selected group. */
 	FText GetGroupComboText() const;
 
+	/** Handle selected tests tag editing. */
+	void HandleEditTagsButtonClicked() const;
+	/** Whether the edit tags button is enabled. */
+	bool IsEditTagsButtonEnabled() const;
+
 	/**
 	 * Handle the copy button clicked in the command bar.
 	 *
@@ -681,4 +688,7 @@ private:
 
 	// Saves the last selected test so that it can be unbinded from refreshing the log
 	TWeakPtr<IAutomationReport> PreviousSelection;
+
+	/** Holds a pointer to the tag filter widget. */
+	TSharedPtr<SAutomationTestTagFilter> TagFilter;
 };
diff --git a/Engine/Source/Programs/AutomationTool/BuildGraph/BgCondition.cs b/Engine/Source/Programs/AutomationTool/BuildGraph/BgCondition.cs
--- a/Engine/Source/Programs/AutomationTool/BuildGraph/BgCondition.cs
+++ b/Engine/Source/Programs/AutomationTool/BuildGraph/BgCondition.cs
@@ -343,6 +343,13 @@ namespace AutomationTool
 
 				result = ContainsItem(arguments.ElementAt(0), arguments.ElementAt(1), arguments.ElementAt(2)) ? "true" : "false";
 			}
+			else if (String.Equals(_tokens[_idx], "Length", StringComparison.OrdinalIgnoreCase) && _tokens[_idx + 1] == "(")
+			{
+				// returns the length of the string
+				_idx++;
+				string argument = await EvaluateScalarAsync(context);
+				result = argument.Length.ToString();
+			}
 			else
 			{
 				// Raw scalar. Remove quotes from strings, and allow literals and simple identifiers to pass through directly.
diff --git a/Engine/Source/Programs/AutomationTool/Gauntlet/Framework/Utils/Gauntlet.HordeReport.cs b/Engine/Source/Programs/AutomationTool/Gauntlet/Framework/Utils/Gauntlet.HordeReport.cs
--- a/Engine/Source/Programs/AutomationTool/Gauntlet/Framework/Utils/Gauntlet.HordeReport.cs
+++ b/Engine/Source/Programs/AutomationTool/Gauntlet/Framework/Utils/Gauntlet.HordeReport.cs
@@ -555,6 +555,7 @@ namespace Gauntlet
 			public class TestResult
 			{
 				public string Name { get; set; }
+				public List<string> Tags { get; set; }
 				public string TestUID { get; set; }
 				public string Suite { get; set; }
 				public TestStateType State { get; set; }
@@ -567,9 +568,10 @@ namespace Gauntlet
 				// not part of json output
 				public List<TestEvent> Events;
 				public SortedSet<string> ErrorHashes;
-				public TestResult(string InName, string InSuite, TestStateType InState, List<string> InDevices, string InDateTime)
+				public TestResult(string InName, List<string> InTags, string InSuite, TestStateType InState, List<string> InDevices, string InDateTime)
 				{
 					Name = InName;
+					Tags = InTags;
 					TestUID = GenerateHash(Name);
 					Suite = InSuite;
 					State = InState;
@@ -632,9 +634,9 @@ namespace Gauntlet
 					Tests = new Dictionary<string, TestResult>();
 					TestResultsTestDataUID = Guid.NewGuid().ToString();
 				}
-				public TestResult NewTestResult(string InName, string InSuite, TestStateType InState, List<string> InDevices, string InDateTime)
+				public TestResult NewTestResult(string InName, List<string> InTags, string InSuite, TestStateType InState, List<string> InDevices, string InDateTime)
 				{
-					TestResult NewItem = new TestResult(InName, InSuite, InState, InDevices, InDateTime);
+					TestResult NewItem = new TestResult(InName, InTags, InSuite, InState, InDevices, InDateTime);
 					Tests[NewItem.TestUID] = NewItem;
 					return NewItem;
 				}
@@ -811,9 +813,9 @@ namespace Gauntlet
 			/// <param name="InSuite"></param>
 			/// <param name="InDevices"></param>
 			/// <param name="InDateTime"></param>
-			public void SetTest(string InName, string InSuite, List<string> InDevices, string InDateTime = null)
+			public void SetTest(string InName, List<string> InTags, string InSuite, List<string> InDevices, string InDateTime = null)
 			{
-				TestResult Test = TestSessionInfo.NewTestResult(InName, InSuite, TestStateType.Unknown, InDevices, InDateTime);
+				TestResult Test = TestSessionInfo.NewTestResult(InName, InTags, InSuite, TestStateType.Unknown, InDevices, InDateTime);
 				CurrentTestUID = Test.TestUID;
 				// Index the test result data.
 				TestResults[Test.TestUID] = new TestResultData(Test);
@@ -959,7 +961,7 @@ namespace Gauntlet
 							// Special case: when UE Test DateTime is set this way, it means the value was 0 or null before getting converted to json.
 							TestDateTime = null;
 						}
-						OutTestPassResults.SetTest(InTestResult.FullTestPath, InSuite, InTestResult.DeviceInstance, InTestResult.DateTime);
+						OutTestPassResults.SetTest(InTestResult.FullTestPath, InTestResult.Tags, InSuite, InTestResult.DeviceInstance, InTestResult.DateTime);
 						foreach (UnrealAutomationEntry InTestEntry in InTestResult.Entries)
 						{
 							string Tag = "entry";
@@ -2255,6 +2257,10 @@ namespace Gauntlet
 					foreach (UnrealAutomatedTestResult InTestResult in InTestPassResults.Tests)
 					{
 						TestPhase Phase = AddPhase(InTestResult.FullTestPath);
+						if (InTestResult.Tags.Count > 0)
+						{
+							Phase.AddTags(InTestResult.Tags);
+						}
 						Phase.deviceKeys = InTestResult.DeviceInstance.ToHashSet();
 						string TestDateTime = InTestResult.DateTime;
 						DateTime PhaseTime = ConvertUETimeStringToDateTime(TestDateTime);
diff --git a/Engine/Source/Programs/AutomationTool/Gauntlet/Unreal/Automation/UE.Automation.cs b/Engine/Source/Programs/AutomationTool/Gauntlet/Unreal/Automation/UE.Automation.cs
--- a/Engine/Source/Programs/AutomationTool/Gauntlet/Unreal/Automation/UE.Automation.cs
+++ b/Engine/Source/Programs/AutomationTool/Gauntlet/Unreal/Automation/UE.Automation.cs
@@ -369,6 +369,11 @@ namespace UE
 				AppConfig.CommandLine += string.Format(" -LogCmds=\"{0}\"", LogCmdsString);
 			}
 
+			if (!string.IsNullOrEmpty(TagFilter))
+			{
+				AppConfig.CommandLine += string.Format(" -ExecCmds=\"Automation SetTagFilter {0}\"", TagFilter);
+			}
+
 		}
 	}
 
diff --git a/Engine/Source/Programs/AutomationTool/Gauntlet/Unreal/Base/Gauntlet.UnrealTestConfiguration.cs b/Engine/Source/Programs/AutomationTool/Gauntlet/Unreal/Base/Gauntlet.UnrealTestConfiguration.cs
--- a/Engine/Source/Programs/AutomationTool/Gauntlet/Unreal/Base/Gauntlet.UnrealTestConfiguration.cs
+++ b/Engine/Source/Programs/AutomationTool/Gauntlet/Unreal/Base/Gauntlet.UnrealTestConfiguration.cs
@@ -762,6 +762,12 @@ namespace Gauntlet
 		[AutoParam]
 		public string VerboseLogCategories { get; set; }
 
+		/// <summary>
+		/// Apply a tag filter, restricting tests to run
+		/// </summary>
+		[AutoParam]
+		public string TagFilter { get; set; }
+
 		// Member variables 
 
 		/// <summary>
diff --git a/Engine/Source/Programs/AutomationTool/Gauntlet/Unreal/Base/Gauntlet.UnrealTestNode.cs b/Engine/Source/Programs/AutomationTool/Gauntlet/Unreal/Base/Gauntlet.UnrealTestNode.cs
--- a/Engine/Source/Programs/AutomationTool/Gauntlet/Unreal/Base/Gauntlet.UnrealTestNode.cs
+++ b/Engine/Source/Programs/AutomationTool/Gauntlet/Unreal/Base/Gauntlet.UnrealTestNode.cs
@@ -1647,6 +1647,13 @@ namespace Gauntlet
 			Report.SetMetadata("BuildTarget", string.Join("+", AllRoleTypes.Select(R => R.ToString()).Distinct().OrderBy(B => B)));
 			Report.SetMetadata("Configuration", string.Join("+", AllRoleContexts.Select(R => R.Configuration.ToString()).Distinct().OrderBy(C => C)));
 			Report.SetMetadata("Project", Context.BuildInfo.ProjectName);
+
+			var TagFilter = Globals.Params.ParseValue("TagFilter", string.Empty);
+			if (!string.IsNullOrEmpty(TagFilter))
+			{
+				Report.SetMetadata("TagFilter", TagFilter);
+			}
+
 			// Additional metadata passed through the command line arguments
 			foreach (var Meta in Globals.Params.ParseValues("Metadata"))
 			{
diff --git a/Engine/Source/Programs/AutomationTool/Gauntlet/Unreal/RunUnreal.cs b/Engine/Source/Programs/AutomationTool/Gauntlet/Unreal/RunUnreal.cs
--- a/Engine/Source/Programs/AutomationTool/Gauntlet/Unreal/RunUnreal.cs
+++ b/Engine/Source/Programs/AutomationTool/Gauntlet/Unreal/RunUnreal.cs
@@ -235,6 +235,7 @@ namespace Gauntlet
 					// look for -clientexeccmds=, -editorexeccmds= etc, these are separate from clientargs for sanity
 					string ExecCmds = Globals.Params.ParseValue("ExecCmds", "");
 					string RoleExecCmds= Globals.Params.ParseValue(Type.ToString() + "ExecCmds", "");
+					string TagFilter = Globals.Params.ParseValue("TagFilter", "");
 					if (!string.IsNullOrEmpty(ExecCmds))
 					{
 						Role.ExtraArgs += string.Format(" -ExecCmds=\"{0}\"", ExecCmds);
@@ -245,6 +246,11 @@ namespace Gauntlet
 						Role.ExtraArgs += string.Format(" -ExecCmds=\"{0}\"", RoleExecCmds);
 					}
 
+					if (!string.IsNullOrEmpty(TagFilter))
+					{
+						Role.ExtraArgs += string.Format(" -ExecCmds=\"Automation SetTagFilter {0}\"", TagFilter);
+					}
+
 					bool UsesEditor = EditorForAllRoles || Globals.Params.ParseParam("Editor" + Type.ToString());
 
 					if (UsesEditor)
diff --git a/Engine/Source/Programs/AutomationTool/Gauntlet/Unreal/Utils/Gauntlet.UnrealTestResults.cs b/Engine/Source/Programs/AutomationTool/Gauntlet/Unreal/Utils/Gauntlet.UnrealTestResults.cs
--- a/Engine/Source/Programs/AutomationTool/Gauntlet/Unreal/Utils/Gauntlet.UnrealTestResults.cs
+++ b/Engine/Source/Programs/AutomationTool/Gauntlet/Unreal/Utils/Gauntlet.UnrealTestResults.cs
@@ -280,6 +280,8 @@ namespace Gauntlet
 		public string TestDisplayName { get; set; }
 		[JsonPropertyName("fullTestPath")]
 		public string FullTestPath { get; set; }
+		[JsonPropertyName("tags")]
+		public List<string> Tags { get; set; }
 		[JsonPropertyName("artifactName")]
 		public string ArtifactName { get; set; }
 		[JsonPropertyName("state")]
diff --git a/Engine/Source/Programs/AutomationTool/Scripts/CopyBuildToStagingDirectory.Automation.cs b/Engine/Source/Programs/AutomationTool/Scripts/CopyBuildToStagingDirectory.Automation.cs
--- a/Engine/Source/Programs/AutomationTool/Scripts/CopyBuildToStagingDirectory.Automation.cs
+++ b/Engine/Source/Programs/AutomationTool/Scripts/CopyBuildToStagingDirectory.Automation.cs
@@ -1589,6 +1589,12 @@ namespace AutomationScripts
 
 					StageConfigFiles(SC, DirectoryReference.Combine(SC.ProjectRoot, "Config"), null);
 
+					if (!SC.StageTargetConfigurations.Any(Config => Config == UnrealTargetConfiguration.Shipping))
+					{
+						StageTagConfigFiles(SC, "Engine", SC.ProjectRoot.FullName);
+						StageTagConfigFiles(SC, "Engine", SC.EngineRoot.FullName);
+					}
+
 					bool bStageVerseSourceFiles = false;
 					if (SC.DedicatedServer)
 					{
@@ -2153,6 +2159,43 @@ namespace AutomationScripts
 
 		}
 
+		static void StageTagConfigFiles(DeploymentContext SC, string StagedBaseFolder, string BasePath)
+		{
+			List<string> AllRestrictedFolderNames = RestrictedFolder.GetNames().ToList();
+			AllRestrictedFolderNames.AddRange(new List<string>() { "Platforms", SC.StageTargetPlatform.ToString() });
+
+			List<string> AllPaths = new List<string>() { BasePath };
+
+			GetPossibleRestrictedPaths(BasePath, AllRestrictedFolderNames, AllPaths);
+
+			string TagsSettingsDirPath;
+			foreach (var Dir in AllPaths)
+			{
+				TagsSettingsDirPath = Path.Combine(Dir, "Config", "Tests");
+				if (DirectoryExists(TagsSettingsDirPath))
+				{
+					string Allowed = Path.Combine(Path.GetRelativePath(new DirectoryReference(BasePath).ParentDirectory.FullName, TagsSettingsDirPath), "Tags.ini");
+					ConsoleUtils.WriteWarning("Allowing test tags config from restricted path " + Allowed);
+					SC.ConfigFilesAllowList.Add(new StagedFileReference(Allowed));
+					StageConfigFiles(SC, new DirectoryReference(TagsSettingsDirPath), null);
+				}
+			}
+		}
+
+		static void GetPossibleRestrictedPaths(string BasePath, List<string> RestrictedFolderNames, List<string> OutRestrictedFolders)
+		{
+			string PotentialPath;
+			foreach (string RestrictedFolderName in RestrictedFolderNames)
+			{
+				PotentialPath = Path.Combine(BasePath, RestrictedFolderName);
+				if (DirectoryExists(PotentialPath))
+				{
+					OutRestrictedFolders.Add(PotentialPath);
+					GetPossibleRestrictedPaths(PotentialPath, RestrictedFolderNames, OutRestrictedFolders);
+				}
+			}
+		}
+
 		/// <summary>
 		/// Stage the appropriate config files from the given input directory
 		/// </summary>
diff --git a/Engine/Source/Programs/UnrealBuildTool/System/RestrictedFolders.cs b/Engine/Source/Programs/UnrealBuildTool/System/RestrictedFolders.cs
--- a/Engine/Source/Programs/UnrealBuildTool/System/RestrictedFolders.cs
+++ b/Engine/Source/Programs/UnrealBuildTool/System/RestrictedFolders.cs
@@ -194,6 +194,11 @@ namespace UnrealBuildTool
 		/// </summary>
 		public static RestrictedFolder LimitedAccess = FindOrAddByName("LimitedAccess");
 
+		/// <summary>
+		/// Base folder that is usually parent of many restricted folders 
+		/// </summary>
+		public static RestrictedFolder Restricted = FindOrAddByName("Restricted");
+
 		/// <summary>
 		/// Epic Employees and Contractors
 		/// </summary>
diff --git a/Engine/Source/Runtime/Core/Private/Misc/AutomationTest.cpp b/Engine/Source/Runtime/Core/Private/Misc/AutomationTest.cpp
--- a/Engine/Source/Runtime/Core/Private/Misc/AutomationTest.cpp
+++ b/Engine/Source/Runtime/Core/Private/Misc/AutomationTest.cpp
@@ -374,15 +374,49 @@ bool FAutomationTestFramework::UnregisterAutomationTest( const FString& InTestNa
 	return bRegistered;
 }
 
-bool FAutomationTestFramework::RegisterAutomationTestTags( const FString& InTestNameToRegister, const FString& InTestTagsToRegister )
+bool FAutomationTestFramework::RegisterAutomationTestTags( const FString& InTestNameToRegister, const FString& InTestTagsToRegister, bool InImmutable)
 {
-	const bool bAlreadyRegistered = TestFullNameToTagDataMap.Contains( InTestNameToRegister );
-	if ( !bAlreadyRegistered )
+	TArray<FString> TagsToRegister;
+	if (InImmutable)
+	{
+		const bool bAlreadyRegistered = TestFullNameToTagDataMap.Contains(InTestNameToRegister);
+		if (!bAlreadyRegistered)
+		{
+			LLM_SCOPE_BYNAME(TEXT("AutomationTest/Framework"));
+			TestFullNameToTagDataMap.Add(InTestNameToRegister, InTestTagsToRegister);
+
+			InTestTagsToRegister.Replace(TEXT("]"), TEXT("")).ParseIntoArray(TagsToRegister, TEXT("["), true);
+			AllExistingTags.Append(TagsToRegister);
+
+			if (!ImmutableTags.Contains(InTestNameToRegister))
+			{
+				ImmutableTags.Add(InTestNameToRegister, TSet<FString>());
+			}
+			ImmutableTags[InTestNameToRegister].Append(TagsToRegister);
+		}
+		return !bAlreadyRegistered;
+	}
+	else // Merge with existing tags, keep immutable tags
 	{
 		LLM_SCOPE_BYNAME(TEXT("AutomationTest/Framework"));
-		TestFullNameToTagDataMap.Add(InTestNameToRegister, InTestTagsToRegister);
+		InTestTagsToRegister.Replace(TEXT("]"), TEXT("")).ParseIntoArray(TagsToRegister, TEXT("["), true);
+		AllExistingTags.Append(TagsToRegister);
+
+		TSet<FString> MergedTags =
+			TSet<FString>(TagsToRegister)
+			.Union(TSet<FString>(
+				GetTagsForAutomationTestAsArray(InTestTagsToRegister))
+				.Intersect(ImmutableTags.Contains(InTestNameToRegister) ? ImmutableTags[InTestNameToRegister] : TSet<FString>()));
+
+		TestFullNameToTagDataMap.Remove(InTestNameToRegister);
+		TestFullNameToTagDataMap.Add(InTestNameToRegister,
+			FString::JoinBy(MergedTags, TEXT(""),
+			[](const FString& InTag)
+			{
+				return FString::Printf(TEXT("[%s]"), *InTag);
+			}));
+		return true;
 	}
-	return !bAlreadyRegistered;
 }
 
 bool FAutomationTestFramework::UnregisterAutomationTestTags(const FString& InTestNameToUnregister)
@@ -401,6 +435,16 @@ bool FAutomationTestFramework::RegisterComplexAutomationTestTags(const FAutomati
 	return RegisterAutomationTestTags(FullTestName, InTestTagsToRegister);
 }
 
+TSet<FString> FAutomationTestFramework::GetAllExistingTags()
+{
+	return AllExistingTags;
+}
+
+bool FAutomationTestFramework::IsTagImmutable(const FString& InTestName, const FString& InTag) const
+{
+	return ImmutableTags.Contains(InTestName) && ImmutableTags[InTestName].Contains(InTag);
+}
+
 FString FAutomationTestFramework::GetTagsForAutomationTest(const FString& InTestName)
 {
 	FString * FindResult = TestFullNameToTagDataMap.Find(InTestName);
@@ -411,6 +455,17 @@ FString FAutomationTestFramework::GetTagsForAutomationTest(const FString& InTest
 	return FString();
 }
 
+TArray<FString> FAutomationTestFramework::GetTagsForAutomationTestAsArray(const FString& InTestName)
+{
+	TArray<FString> Tags;
+	FString* FindResult = TestFullNameToTagDataMap.Find(InTestName);
+	if (FindResult)
+	{
+		FindResult->Replace(TEXT("]"), TEXT("[")).ParseIntoArray(Tags, TEXT("["));
+	}
+	return Tags;
+}
+
 void FAutomationTestFramework::EnqueueLatentCommand(TSharedPtr<IAutomationLatentCommand> NewCommand)
 {
 	//ensure latent commands are never used within smoke tests - will only catch when smokes are exclusively requested
@@ -850,6 +905,148 @@ void FAutomationTestFramework::GetTestFullNamesMatchingTagPattern(TArray<FString
 	}
 }
 
+void FAutomationTestFramework::GetPossibleRestrictedPaths(const FString& BasePath, const TArray<FString>& RestrictedFolderNames, TArray<FString>& OutRestrictedFolders) const
+{
+	FString PotentialPath;
+	for (const FString& RestrictedFolderName : RestrictedFolderNames)
+	{
+		PotentialPath = BasePath / RestrictedFolderName;
+		if (IFileManager::Get().DirectoryExists(*PotentialPath))
+		{
+			OutRestrictedFolders.Add(PotentialPath);
+			GetPossibleRestrictedPaths(PotentialPath, RestrictedFolderNames, OutRestrictedFolders);
+		}
+	}
+}
+
+void FAutomationTestFramework::LoadTestTagMappings()
+{
+	static bool bTagMappingFilesLoaded = false;
+
+	if (!bTagMappingFilesLoaded)
+	{
+		LoadTestTagMappings(FPaths::ConvertRelativePathToFull(*FPaths::ProjectDir()));
+		LoadTestTagMappings(FPaths::ConvertRelativePathToFull(FPaths::EngineDir()));
+		
+		bTagMappingFilesLoaded = true;
+	}
+}
+
+void FAutomationTestFramework::LoadTestTagMappings(const FString& BasePath)
+{
+	TArray<FString> RestrictedFolderNames = FPaths::GetRestrictedFolderNames();
+	RestrictedFolderNames.Append({ TEXT("Platforms"), TEXT("Restricted") });
+
+	TArray<FString> AllPaths;
+	AllPaths.Add(BasePath);
+
+	GetPossibleRestrictedPaths(BasePath, RestrictedFolderNames, AllPaths);
+
+	FString TagsSettingsFilePath;
+	TArray<FString> TestTagMappings;
+
+	for (const FString& Path : AllPaths)
+	{
+		TagsSettingsFilePath = Path / TEXT("Config") / TEXT("Tests") / TEXT("Tags.ini");
+
+		if (IFileManager::Get().FileExists(*TagsSettingsFilePath))
+		{
+			UE_LOG(LogAutomationTest, Log, TEXT("Loading tags from %s."), *TagsSettingsFilePath);
+			GConfig->LoadFile(TagsSettingsFilePath);
+
+			GConfig->GetSection(TEXT("TestTagMapping"), TestTagMappings, TagsSettingsFilePath);
+			FString TestName, TestTags;
+			for (FString MapLine : TestTagMappings)
+			{
+				MapLine.Split(TEXT("="), &TestName, &TestTags);
+				RegisterAutomationTestTags(TestName, TestTags, false);
+			}
+		}
+	}
+}
+
+void FAutomationTestFramework::SaveTestTagMappings(const TArray<FString>& TestFullNames, const TArray<FString>& TestFilePaths) const
+{
+	const FString ProjectDirPath = FPaths::ConvertRelativePathToFull(*FPaths::ProjectDir());
+	const FString EngineDirPath = FPaths::ConvertRelativePathToFull(FPaths::EngineDir());
+
+	TArray<FString> RestrictedFolderNames = FPaths::GetRestrictedFolderNames();
+	RestrictedFolderNames.Append({ TEXT("Platforms"), TEXT("Restricted") });
+
+	TArray<FString> RestrictedPathsProject;
+	GetPossibleRestrictedPaths(ProjectDirPath, RestrictedFolderNames, RestrictedPathsProject);
+	TArray<FString> RestrictedPathsEngine;
+	GetPossibleRestrictedPaths(EngineDirPath, RestrictedFolderNames, RestrictedPathsEngine);
+
+	RestrictedPathsProject.Sort([](const FString& A, const FString& B) { return A.Len() > B.Len(); });
+	RestrictedPathsEngine.Sort([](const FString& A, const FString& B) { return A.Len() > B.Len(); });
+
+	FString TestSourcePath;
+	FString TagsSettingsFilePath;
+	TSet<FString> AllTagConfigPaths;
+	for (int i = 0; i < TestFilePaths.Num(); i++)
+	{
+		TagsSettingsFilePath.Empty();
+		TestSourcePath = TestFilePaths[i];
+		// Pick longest path with restricted names that is a suffix of the test path
+		for (int j = 0 ; j < RestrictedPathsProject.Num(); j++)
+		{
+			if (FPaths::IsUnderDirectory(TestSourcePath, RestrictedPathsProject[j]))
+			{
+				TagsSettingsFilePath = RestrictedPathsProject[j] / TEXT("Config") / TEXT("Tests") / TEXT("Tags.ini");
+				break;
+			}
+		}
+
+		if (TagsSettingsFilePath.IsEmpty())
+		{
+			for (int j = 0; j < RestrictedPathsEngine.Num(); j++)
+			{
+				if (FPaths::IsUnderDirectory(TestSourcePath, RestrictedPathsEngine[j]))
+				{
+					TagsSettingsFilePath = RestrictedPathsEngine[j] / TEXT("Config") / TEXT("Tests") / TEXT("Tags.ini");
+					break;
+				}
+			}
+		}
+
+		if (TagsSettingsFilePath.IsEmpty())
+		{
+			UE_LOG(LogAutomationTest, Error, TEXT("Test neither in project path or engine path, cannot be mapped. Path: %s"), *TestSourcePath);
+			continue;
+		}
+
+		if (!IFileManager::Get().FileExists(*TagsSettingsFilePath))
+		{
+			TUniquePtr<FArchive> Ar = TUniquePtr<FArchive>(IFileManager::Get().CreateFileWriter(*TagsSettingsFilePath));
+			if (!Ar)
+			{
+				UE_LOG(LogAutomationTest, Warning, TEXT("Failed to create test tags mapping file."));
+			}
+			else
+			{
+				UE_LOG(LogAutomationTest, Log, TEXT("Created new test tags mapping file: %s."), *TagsSettingsFilePath);
+
+				const FString TagMappingHeader = TEXT("[TestTagMapping]");
+				Ar->Serialize(TCHAR_TO_ANSI(*TagMappingHeader), TagMappingHeader.Len());
+				Ar->Close();
+			}
+		}
+
+		if (!AllTagConfigPaths.Contains(TagsSettingsFilePath))
+		{
+			AllTagConfigPaths.Add(TagsSettingsFilePath);
+		}
+		GConfig->SetString(TEXT("TestTagMapping"), *TestFullNames[i], **TestFullNameToTagDataMap.Find(TestFullNames[i]), TagsSettingsFilePath);
+	}
+
+	for (FString ConfigPath : AllTagConfigPaths)
+	{
+		UE_LOG(LogAutomationTest, Log, TEXT("Saving tags to %s."), *ConfigPath);
+		GConfig->Flush(false, ConfigPath);
+	}
+}
+
 bool FAutomationTestFramework::ShouldTestContent(const FString& Path) const
 {
 	static TArray<FString> TestLevelFolders;
diff --git a/Engine/Source/Runtime/Core/Public/Misc/AutomationTest.h b/Engine/Source/Runtime/Core/Public/Misc/AutomationTest.h
--- a/Engine/Source/Runtime/Core/Public/Misc/AutomationTest.h
+++ b/Engine/Source/Runtime/Core/Public/Misc/AutomationTest.h
@@ -1020,11 +1020,11 @@ public:
 	 * 
 	 * @param	InTestNameToRegister	FullName of the test to associate with tags
 	 * @param	InTestTagsToRegister	Concatenated string of tags to register
+	 * @param   InImmutable				Whether tags are immutable. If false, merge with existing tags, keeping the immutable
 	 * 
-	 * @return	true if tags for the test successfully registered; false if tags were already registered for the test
-	 * @see FAutomationTestBase::GetTestFullName
+	 * @return	true if tags for the test successfully registered or overwritten; false if tags were already registered for the test
 	 */
-	CORE_API bool RegisterAutomationTestTags(const FString& InTestNameToRegister, const FString& InTestTagsToRegister);
+	CORE_API bool RegisterAutomationTestTags(const FString& InTestNameToRegister, const FString& InTestTagsToRegister, bool InImmutable = true);
 
 	/**
 	 * Unregister tags for a automation test from the framework.
@@ -1032,7 +1032,6 @@ public:
 	 * @param	InTestNameToUnregister	FullName of the test to remove tags from
 	 * 
 	 * @return true if test tags were successfully unregistered; false if tags for a test with that name was not found in the framework.
-	 * @see FAutomationTestBase::GetTestFullName
 	 */
 	CORE_API bool UnregisterAutomationTestTags(const FString& InTestNameToUnregister);
 
@@ -1047,6 +1046,10 @@ public:
 	 */
 	CORE_API bool RegisterComplexAutomationTestTags(const FAutomationTestBase* InTest, const FString& InBeautifiedTestName, const FString& InTestTagsToRegister);
 
+	CORE_API TSet<FString> GetAllExistingTags();
+
+	CORE_API bool IsTagImmutable(const FString& InTestName, const FString& InTag) const;
+
 	/**
 	 * Fetch the tags associated with a test
 	 * 
@@ -1056,6 +1059,15 @@ public:
 	 */
 	CORE_API FString GetTagsForAutomationTest(const FString& InTestName);
 
+	/**
+	 * Fetch the tags associated with a test as an array of strings
+	 *
+	 * @param	InTestName	Name of the test to find tags for
+	 *
+	 * @return	an array of tags; empty if test is not registered
+	 */
+	CORE_API TArray<FString> GetTagsForAutomationTestAsArray(const FString& InTestName);
+
 	/**
 	 * Enqueues a latent command for execution on a subsequent frame
 	 *
@@ -1158,6 +1170,23 @@ public:
 	 */
 	CORE_API void GetTestFullNamesMatchingTagPattern(TArray<FString>& OutTestNames, const FString& TagPattern) const;
 
+	CORE_API void GetPossibleRestrictedPaths(const FString& BasePath, const TArray<FString>& RestrictedFolderNames, TArray<FString>& OutRestrictedFolders) const;
+
+	/**
+	 * Reads test tag mappings from all Tags.ini config files
+	 */
+	CORE_API void LoadTestTagMappings();
+
+	/**
+	 * Reads test tag mappings under a certain base path
+	 */
+	CORE_API void LoadTestTagMappings(const FString& BasePath);
+
+	/**
+	 * Saves runtime edited test tag mappings to Tags.ini
+	 */
+	CORE_API void SaveTestTagMappings(const TArray<FString>& TestFullNames, const TArray<FString>& TestFilePaths) const;
+
 	/**
 	 * Whether the testing framework should allow content to be tested or not.  Intended to block developer directories.
 	 * @param Path - Full path to the content in question
@@ -1441,9 +1470,19 @@ private:
 	/** Mapping of automation test names to their respective object instances */
 	TMap<FString, FAutomationTestBase*> AutomationTestClassNameToInstanceMap;
 
-	/** Mapping of full test names to their registered tags */
+	/* Set of all tags whether they are code declared or through config file */
+	TSet<FString> AllExistingTags;
+
+	/**
+	 * Mapping of full test names to their registered tags.
+	 */
 	TMap<FString, FString> TestFullNameToTagDataMap;
 
+	/**
+	 * Map of immutable tags for each test. Code defined tags are immutable.
+	 */
+	TMap<FString, TSet<FString>> ImmutableTags;
+
 	/** Queue of deferred commands */
 	TQueue< TSharedPtr<IAutomationLatentCommand> > LatentCommands;
 
diff --git a/Engine/Source/Runtime/Core/Public/Misc/LowLevelTestAdapter.h b/Engine/Source/Runtime/Core/Public/Misc/LowLevelTestAdapter.h
--- a/Engine/Source/Runtime/Core/Public/Misc/LowLevelTestAdapter.h
+++ b/Engine/Source/Runtime/Core/Public/Misc/LowLevelTestAdapter.h
@@ -48,14 +48,14 @@ FString CaptureExpressionsAndValues(const FString& InExpressions, ArgTypes&&...
 }
 
 
-#define IMPLEMENT_SIMPLE_AUTOMATION_TEST_PRIVATE_LLT( TClass, PrettyName, TFlags, FileName, LineNumber ) \
+#define IMPLEMENT_SIMPLE_AUTOMATION_TEST_PRIVATE_LLT( TClass, PrettyName, TFlagsOrTags, FileName, LineNumber ) \
 		class TClass : public FAutomationTestBase \
 		{ \
 		public:\
 			TClass( const FString& InName) \
 			: FAutomationTestBase( InName, false ) \
 			{ \
-				TestFlags = ExtractAutomationTestFlags(TFlags); \
+				TestFlags = ExtractAutomationTestFlags(TFlagsOrTags); \
 				PrettyNameDotNotation = FString(PrettyName).Replace(TEXT("::"), TEXT(".")); \
 				if (!(TestFlags & EAutomationTestFlags_ApplicationContextMask)) \
 				{ \
@@ -65,6 +65,7 @@ FString CaptureExpressionsAndValues(const FString& InExpressions, ArgTypes&&...
 				{ \
 					TestFlags |= EAutomationTestFlags::EngineFilter; \
 				} \
+				FAutomationTestFramework::Get().RegisterAutomationTestTags(GetBeautifiedTestName(), TFlagsOrTags); \
 			} \
 			virtual EAutomationTestFlags GetTestFlags() const override { return TestFlags; } \
 			virtual bool IsStressTest() const { return false; } \
@@ -91,8 +92,8 @@ FString CaptureExpressionsAndValues(const FString& InExpressions, ArgTypes&&...
 #define LLT_JOIN(Prefix, Counter) LLT_JOIN_INNER(Prefix, Counter)
 #define LLT_JOIN_INNER(Prefix, Counter) Prefix##Counter
 
-#define TEST_CASE_NAMED_STR(TClass, StrName, PrettyName, TFlags) \
-		IMPLEMENT_SIMPLE_AUTOMATION_TEST_PRIVATE_LLT(TClass, PrettyName, TFlags, __FILE__, __LINE__) \
+#define TEST_CASE_NAMED_STR(TClass, StrName, PrettyName, TFlagsOrTags) \
+		IMPLEMENT_SIMPLE_AUTOMATION_TEST_PRIVATE_LLT(TClass, PrettyName, TFlagsOrTags, __FILE__, __LINE__) \
 		namespace \
 		{ \
 			TClass LLT_JOIN(TClass, Instance)(TEXT(StrName)); \
@@ -105,8 +106,8 @@ FString CaptureExpressionsAndValues(const FString& InExpressions, ArgTypes&&...
 #define LLT_STR_EXPAND(Macro) LLT_STR(Macro)
 #define TEST_CASE_GENERATED_NAME_UNIQUE_STR LLT_STR_EXPAND(TEST_CASE_GENERATED_NAME_UNIQUE)
 // Note: TEST_CASE uses unique names which only work when used inside unique namespace in the same compilation unit?// Use TEST_CASE_NAMED instead and provide an unique global instance name
-#define TEST_CASE(PrettyName, TFlags) TEST_CASE_NAMED_STR(TEST_CASE_GENERATED_NAME_UNIQUE, TEST_CASE_GENERATED_NAME_UNIQUE_STR, PrettyName, TFlags)
-#define TEST_CASE_NAMED(ClassName, PrettyName, TFlags) TEST_CASE_NAMED_STR(ClassName, #ClassName, PrettyName, TFlags)
+#define TEST_CASE(PrettyName, TFlagsOrTags) TEST_CASE_NAMED_STR(TEST_CASE_GENERATED_NAME_UNIQUE, TEST_CASE_GENERATED_NAME_UNIQUE_STR, PrettyName, TFlagsOrTags)
+#define TEST_CASE_NAMED(ClassName, PrettyName, TFlagsOrTags) TEST_CASE_NAMED_STR(ClassName, #ClassName, PrettyName, TFlagsOrTags)
 
 // Both python and oodle don't trust __LINE__ for unique names, and use __COUNTER__ where possible
 #ifdef __COUNTER__
diff --git a/Engine/Source/Runtime/Slate/Private/Widgets/Input/SSuggestionTextBox.cpp b/Engine/Source/Runtime/Slate/Private/Widgets/Input/SSuggestionTextBox.cpp
--- a/Engine/Source/Runtime/Slate/Private/Widgets/Input/SSuggestionTextBox.cpp
+++ b/Engine/Source/Runtime/Slate/Private/Widgets/Input/SSuggestionTextBox.cpp
@@ -29,6 +29,8 @@ void SSuggestionTextBox::Construct( const FArguments& InArgs )
 	OnTextChanged = InArgs._OnTextChanged;
 	OnTextCommitted = InArgs._OnTextCommitted;
 
+	CommitsOnSuggestionSelected = InArgs._CommitOnSuggestionSelected.Get();
+
 	FSlateApplication::Get().OnFocusChanging().AddSP(this, &SSuggestionTextBox::OnGlobalFocusChanging);
 
 	ChildSlot
@@ -232,6 +234,10 @@ FReply SSuggestionTextBox::OnKeyDown( const FGeometry& MyGeometry, const FKeyEve
 
 			return FReply::Handled();
 		}
+		else if (KeyEvent.GetKey() == EKeys::Enter && CommitsOnSuggestionSelected && SuggestionListView->GetNumItemsSelected() > 0)
+		{
+			HandleTextBoxTextCommitted(FText::FromString(*Suggestions[SelectedSuggestion]), ETextCommit::OnEnter);
+		}
 		else if (KeyEvent.GetKey() == EKeys::Tab)
 		{
 			// auto-complete the highlighted suggestion
@@ -335,6 +341,12 @@ void SSuggestionTextBox::HandleSuggestionListViewSelectionChanged( TSharedPtr<FS
 				SelectedSuggestion = i;
 
 				MarkActiveSuggestion();
+
+				if (CommitsOnSuggestionSelected)
+				{
+					HandleTextBoxTextCommitted(FText::FromString(*NewValue.Get()), ETextCommit::OnEnter);
+				}
+
 				FocusTextBox();
 
 				break;
@@ -343,7 +355,6 @@ void SSuggestionTextBox::HandleSuggestionListViewSelectionChanged( TSharedPtr<FS
 	}
 }
 
-
 FText SSuggestionTextBox::HandleSuggestionListWidgetHighlightText( ) const
 {
 	return TextBox->GetText();
diff --git a/Engine/Source/Runtime/Slate/Public/Widgets/Input/SSuggestionTextBox.h b/Engine/Source/Runtime/Slate/Public/Widgets/Input/SSuggestionTextBox.h
--- a/Engine/Source/Runtime/Slate/Public/Widgets/Input/SSuggestionTextBox.h
+++ b/Engine/Source/Runtime/Slate/Public/Widgets/Input/SSuggestionTextBox.h
@@ -49,6 +49,7 @@ public:
 		: _BackgroundColor(FLinearColor::White)
 		, _BackgroundImage(FCoreStyle::Get().GetBrush("SuggestionTextBox.Background"))
 		, _ClearKeyboardFocusOnCommit(true)
+		, _CommitOnSuggestionSelected(false)
 		, _ErrorReporting()
 		, _Font( FCoreStyle::Get().GetFontStyle( TEXT( "NormalFont" ) ) )
 		, _ForegroundColor( FCoreStyle::Get().GetSlateColor( "InvertedForeground" ) )
@@ -73,6 +74,9 @@ public:
 		/** Whether to clear keyboard focus when pressing enter to commit changes. */
 		SLATE_ATTRIBUTE(bool, ClearKeyboardFocusOnCommit)
 
+		/** Whether to commit text that was selected from the suggestions list. */
+		SLATE_ATTRIBUTE(bool, CommitOnSuggestionSelected)
+
 		/** Provides an alternative mechanism for error reporting. */
 		SLATE_ARGUMENT(TSharedPtr<class IErrorReportingWidget>, ErrorReporting)
 
@@ -250,6 +254,9 @@ private:
 	// The VerticalBox containing the SuggestionList.
 	TSharedPtr<SVerticalBox> VerticalBox;
 
+	// Trigger commit text event as soon as suggestion was selected.
+	bool CommitsOnSuggestionSelected;
+
 private:
 
 	// Holds a delegate that is executed before the input history is shown.
