Tag editing: functional test path must be asset path instead of base class path
#rnx

diff --git a/Engine/Source/Developer/AutomationWindow/Private/SAutomationTestTagEditor.h b/Engine/Source/Developer/AutomationWindow/Private/SAutomationTestTagEditor.h
--- a/Engine/Source/Developer/AutomationWindow/Private/SAutomationTestTagEditor.h
+++ b/Engine/Source/Developer/AutomationWindow/Private/SAutomationTestTagEditor.h
@@ -5,6 +5,7 @@
 #include "CoreMinimal.h"
 #include "Misc/CoreDelegates.h"
 #include "Misc/MessageDialog.h"
+#include "Misc/PackageName.h"
 #include "Widgets/DeclarativeSyntaxSupport.h"
 #include "Widgets/SWidget.h"
 #include "Widgets/SCompoundWidget.h"
@@ -325,7 +326,15 @@ private:
 				for (IAutomationReportPtr Report : TestReports)
 				{
 					TestNames.Add(Report->GetFullTestPath());
-					TestFilePaths.Add(Report->GetSourceFile());
+					if (!Report->GetAssetPath().IsEmpty())
+					{
+						FString PackageLocalPath = FPackageName::LongPackageNameToFilename(*Report->GetAssetPath(), FPackageName::GetAssetPackageExtension());
+						TestFilePaths.Add(PackageLocalPath);
+					}
+					else
+					{
+						TestFilePaths.Add(Report->GetSourceFile());
+					}
 					Tags = Report->GetTagsArray();
 					for (FString DeleteTag : MarkedForDeleteTags)
 					{
