Update FunctionalTest UI when tags are edited from the Automation Window and viceversa

#rnx

diff --git a/Engine/Source/Developer/AutomationWindow/Private/SAutomationTestTagEditor.h b/Engine/Source/Developer/AutomationWindow/Private/SAutomationTestTagEditor.h
--- a/Engine/Source/Developer/AutomationWindow/Private/SAutomationTestTagEditor.h
+++ b/Engine/Source/Developer/AutomationWindow/Private/SAutomationTestTagEditor.h
@@ -76,6 +76,7 @@ public:
 		{
 			return SNew(STextBlock)
 				.Text(FText::FromString(TestReport->GetTags()))
+				.Margin(FMargin(10, 0, 0, 0))
 				.ToolTipText(FText::FromString(TestReport->GetTags()));
 		}
 		return SNullWidget::NullWidget;
@@ -180,13 +181,28 @@ public:
 							.AutoHeight()
 							.Padding(5.0f)
 							[
-								SAssignNew(TagSuggestionBox, SSuggestionTextBox)
-									.ForegroundColor(FSlateColor::UseForeground())
-									.CommitOnSuggestionSelected(true)
-									.ClearKeyboardFocusOnCommit(false)
-									.OnTextCommitted(this, &SAutomationTestTagEditor::TagSuggestionBox_SetText)
-									.HintText(LOCTEXT("EnterTagHint", "Enter tag"))
-									.OnShowingSuggestions(this, &SAutomationTestTagEditor::TagSuggestionBox_GetSuggestions)
+								SNew(SHorizontalBox)
+									+ SHorizontalBox::Slot()
+									.FillWidth(0.9f)
+									.VAlign(VAlign_Center)
+									[
+										SAssignNew(TagSuggestionBox, SSuggestionTextBox)
+											.ForegroundColor(FSlateColor::UseForeground())
+											.CommitOnSuggestionSelected(true)
+											.ClearKeyboardFocusOnCommit(false)
+											.OnTextCommitted(this, &SAutomationTestTagEditor::TagSuggestionBox_SetText)
+											.HintText(LOCTEXT("EnterTagHint", "Enter tag"))
+											.OnShowingSuggestions(this, &SAutomationTestTagEditor::TagSuggestionBox_GetSuggestions)
+									]
+									+ SHorizontalBox::Slot()
+									.FillWidth(0.1f)
+									.Padding(FMargin(1, 0, 0, 0))
+									.VAlign(VAlign_Center)
+									[
+										SNew(SButton)
+											.Text(LOCTEXT("AddTags", "Add"))
+											.OnClicked(this, &SAutomationTestTagEditor::OnAddButtonClicked)
+									]
 							]
 							+ SVerticalBox::Slot()
 							.AutoHeight()
@@ -259,12 +275,27 @@ public:
 	}
 
 private:
+	void AddTagToListBox(const FString& InTag)
+	{
+		TagListBox->AddTag(InTag, true, false);
+		TagSuggestionBox->SetText(FText::GetEmpty());
+	}
+
+	FReply OnAddButtonClicked()
+	{
+		FString SelectedTag = TagSuggestionBox->GetText().ToString();
+		if (!SelectedTag.IsEmpty())
+		{
+			AddTagToListBox(SelectedTag);
+		}
+		return FReply::Handled();
+	}
+
 	void TagSuggestionBox_SetText(const FText& NewText, ETextCommit::Type CommitInfo)
 	{
 		if (CommitInfo == ETextCommit::Type::OnEnter && !NewText.ToString().IsEmpty())
 		{
-			TagListBox->AddTag(NewText.ToString(), true, false);
-			TagSuggestionBox->SetText(FText::GetEmpty());
+			AddTagToListBox(NewText.ToString());
 		}
 	}
 
