Bughawk fixes

#rnx

diff --git a/Engine/Source/Developer/AutomationWindow/Private/SAutomationTestTagEditor.h b/Engine/Source/Developer/AutomationWindow/Private/SAutomationTestTagEditor.h
--- a/Engine/Source/Developer/AutomationWindow/Private/SAutomationTestTagEditor.h
+++ b/Engine/Source/Developer/AutomationWindow/Private/SAutomationTestTagEditor.h
@@ -193,6 +193,7 @@ public:
 											.ForegroundColor(FSlateColor::UseForeground())
 											.CommitOnSuggestionSelected(true)
 											.ClearKeyboardFocusOnCommit(false)
+											.OnTextChanged(this, &SAutomationTestTagEditor::TagSuggestionBox_TextChanged)
 											.OnTextCommitted(this, &SAutomationTestTagEditor::TagSuggestionBox_SetText)
 											.HintText(LOCTEXT("EnterTagHint", "Enter tag"))
 											.OnShowingSuggestions(this, &SAutomationTestTagEditor::TagSuggestionBox_GetSuggestions)
@@ -202,8 +203,9 @@ public:
 									.Padding(FMargin(1, 0, 0, 0))
 									.VAlign(VAlign_Center)
 									[
-										SNew(SButton)
+										SAssignNew(AddButton, SButton)
 											.Text(LOCTEXT("AddTags", "Add"))
+											.IsEnabled(false)
 											.OnClicked(this, &SAutomationTestTagEditor::OnAddButtonClicked)
 									]
 							]
@@ -302,6 +304,11 @@ private:
 		}
 	}
 
+	void TagSuggestionBox_TextChanged(const FText& NewText)
+	{
+		AddButton->SetEnabled(!NewText.IsEmpty());
+	}
+
 	void TagSuggestionBox_GetSuggestions(const FString& CurrText, TArray<FString>& OutSuggestions)
 	{
 		OutSuggestions = FAutomationTestFramework::Get().GetAllExistingTags().Array().FilterByPredicate([&CurrText](const FString& CurrentTag) { return CurrentTag.Contains(CurrText); });
@@ -430,6 +437,7 @@ private:
 
 	TSharedPtr<SSuggestionTextBox> TagSuggestionBox;
 	TSharedPtr<SAutomationTestTagListBox> TagListBox;
+	TSharedPtr<SButton> AddButton;
 	TArray<IAutomationReportPtr> TestReports;
 	FOnTestTagsSavedEvent OnTestTagsSavedEvent;
 };
diff --git a/Engine/Source/Runtime/Core/Private/Misc/AutomationTest.cpp b/Engine/Source/Runtime/Core/Private/Misc/AutomationTest.cpp
--- a/Engine/Source/Runtime/Core/Private/Misc/AutomationTest.cpp
+++ b/Engine/Source/Runtime/Core/Private/Misc/AutomationTest.cpp
@@ -1058,7 +1058,7 @@ void FAutomationTestFramework::SaveTestTagMappings(const TArray<FString>& TestFu
 	{
 		UE_LOG(LogAutomationTestFramework, Log, TEXT("Saving tags to %s."), *ConfigPath);
 		AffectedTagConfigPaths[ConfigPath].Write(ConfigPath);
-		AfterConfigSaved.ExecuteIfBound(TagsSettingsFilePath);
+		AfterConfigSaved.ExecuteIfBound(ConfigPath);
 
 	}
 }
diff --git a/Engine/Source/Runtime/Core/Public/Misc/AutomationTest.h b/Engine/Source/Runtime/Core/Public/Misc/AutomationTest.h
--- a/Engine/Source/Runtime/Core/Public/Misc/AutomationTest.h
+++ b/Engine/Source/Runtime/Core/Public/Misc/AutomationTest.h
@@ -1189,7 +1189,7 @@ public:
 
 
 	/**
-	 * Return the list of test tag config files affected by editing as a map of config path and FConfigFile object
+	 * Saves test tag mappings optionally invoking before / after events.
 	 */
 	CORE_API void SaveTestTagMappings(const TArray<FString>& TestFullNames, const TArray<FString>& TestFilePaths, const FBeforeTagMappingConfigSaved& BeforeConfigSaved, const FAfterTagMappingConfigSaved& AfterConfigSaved) const;
 
