Tag ini file - Submit Tool linter

#rnx

diff --git a/Engine/Source/Developer/AutomationWindow/Private/SAutomationTestTagEditor.h b/Engine/Source/Developer/AutomationWindow/Private/SAutomationTestTagEditor.h
--- a/Engine/Source/Developer/AutomationWindow/Private/SAutomationTestTagEditor.h
+++ b/Engine/Source/Developer/AutomationWindow/Private/SAutomationTestTagEditor.h
@@ -369,8 +369,11 @@ private:
 					TestNames.Add(Report->GetFullTestPath());
 					if (!Report->GetAssetPath().IsEmpty())
 					{
-						FString PackageLocalPath = FPackageName::LongPackageNameToFilename(*Report->GetAssetPath(), FPackageName::GetAssetPackageExtension());
-						TestFilePaths.Add(PackageLocalPath);
+						FString PackageLocalPath;
+						if (FPackageName::TryConvertLongPackageNameToFilename(*Report->GetAssetPath(), PackageLocalPath, FPackageName::GetAssetPackageExtension()))
+						{
+							TestFilePaths.Add(PackageLocalPath);
+						}
 					}
 					else
 					{
