BugHawk fixes

#rnx

diff --git a/Engine/Source/Developer/AutomationWindow/Private/SAutomationTestTagEditor.h b/Engine/Source/Developer/AutomationWindow/Private/SAutomationTestTagEditor.h
--- a/Engine/Source/Developer/AutomationWindow/Private/SAutomationTestTagEditor.h
+++ b/Engine/Source/Developer/AutomationWindow/Private/SAutomationTestTagEditor.h
@@ -35,6 +35,8 @@
 
 #define LOCTEXT_NAMESPACE "SAutomationTestTagEditor"
 
+DEFINE_LOG_CATEGORY_STATIC(LogAutomationTestTagEditor, Log, All);
+
 namespace TestNameAndTagTableConstants
 {
 	const FName Name(TEXT("Name"));
@@ -366,7 +368,6 @@ private:
 				TArray<FString> TestFilePaths;
 				for (IAutomationReportPtr Report : TestReports)
 				{
-					TestNames.Add(Report->GetFullTestPath());
 					if (!Report->GetAssetPath().IsEmpty())
 					{
 						FString PackageLocalPath;
@@ -374,11 +375,18 @@ private:
 						{
 							TestFilePaths.Add(PackageLocalPath);
 						}
+						else
+						{
+							UE_LOG(LogAutomationTestTagEditor, Error, TEXT("Could not resolve asset path for test %s, skipping from tag mapping"), *Report->GetFullTestPath());
+							continue;
+						}
 					}
 					else
 					{
 						TestFilePaths.Add(Report->GetSourceFile());
 					}
+					TestNames.Add(Report->GetFullTestPath());
+
 					Tags = Report->GetTagsArray();
 					for (FString DeleteTag : MarkedForDeleteTags)
 					{
diff --git a/Engine/Source/Runtime/Core/Private/Misc/AutomationTest.cpp b/Engine/Source/Runtime/Core/Private/Misc/AutomationTest.cpp
--- a/Engine/Source/Runtime/Core/Private/Misc/AutomationTest.cpp
+++ b/Engine/Source/Runtime/Core/Private/Misc/AutomationTest.cpp
@@ -974,6 +974,8 @@ void FAutomationTestFramework::LoadTestTagMappings(const FString& BasePath, bool
 
 void FAutomationTestFramework::SaveTestTagMappings(const TArray<FString>& TestFullNames, const TArray<FString>& TestFilePaths, const FBeforeTagMappingConfigSaved& BeforeConfigSaved, const FAfterTagMappingConfigSaved& AfterConfigSaved) const
 {
+	check(TestFullNames.Num() == TestFilePaths.Num());
+
 	const FString ProjectDirPath = FPaths::ConvertRelativePathToFull(*FPaths::ProjectDir());
 	const FString EngineDirPath = FPaths::ConvertRelativePathToFull(FPaths::EngineDir());
 
