Fix the ability to add the same test tag after deletion.

#rnx

diff --git a/Engine/Source/Developer/AutomationWindow/Private/SAutomationTestTagListBox.h b/Engine/Source/Developer/AutomationWindow/Private/SAutomationTestTagListBox.h
--- a/Engine/Source/Developer/AutomationWindow/Private/SAutomationTestTagListBox.h
+++ b/Engine/Source/Developer/AutomationWindow/Private/SAutomationTestTagListBox.h
@@ -20,6 +20,8 @@
 
 #define LOCTEXT_NAMESPACE "SAutomationTestTagListBox"
 
+DEFINE_LOG_CATEGORY_STATIC(LogAutomationTestTagListBox, Log, All);
+
 class SAutomationTestTagItem
 	: public SCompoundWidget
 {
@@ -95,19 +97,7 @@ public:
 	}
 
 protected:
-	FReply ToggleTagDeletion()
-	{
-		MarkedForDeletion = !MarkedForDeletion;
-		if (!IsNew)
-		{
-			MainBorder->SetBorderBackgroundColor(MarkedForDeletion ? FColorList::Red : FColorList::Black);
-		}
-		else
-		{
-			ChildSlot.DetachWidget();
-		}
-		return FReply::Handled();
-	}
+	FReply ToggleTagDeletion();
 
 	TSharedPtr<SBorder> MainBorder;
 	bool IsNew;
@@ -189,6 +179,11 @@ public:
 		return NewTags;
 	}
 
+	int32 RemoveNewTagsItem(const FString &ItemToRemove)
+	{
+		return NewTags.RemoveSingle(ItemToRemove);
+	}
+
 	TArray<FString> GetMarkForDeleteTags()
 	{
 		TArray<FString> ToDelete;
@@ -212,4 +207,40 @@ private:
 };
 
 
+inline FReply SAutomationTestTagItem::ToggleTagDeletion()
+{
+	TSharedPtr<SWidget> TagListBoxWidget = MainBorder;
+
+	do
+	{
+		TagListBoxWidget = TagListBoxWidget->GetParentWidget();
+
+		if (TagListBoxWidget == nullptr)
+		{
+			UE_LOG(LogAutomationTestTagListBox, Error, TEXT("Could not find parent widget with type SAutomationTestTagListBox"));
+
+			return FReply::Handled();
+		}
+
+	} while (!TagListBoxWidget->GetTypeAsString().Equals("SAutomationTestTagListBox"));
+
+	const TSharedPtr<SAutomationTestTagListBox> TagListBox = StaticCastSharedPtr<SAutomationTestTagListBox>(TagListBoxWidget);
+	
+	if (TagListBox->RemoveNewTagsItem(TagValue) == 0)
+	{
+		MarkedForDeletion = !MarkedForDeletion;
+	}
+	
+	if (!IsNew)
+	{
+		MainBorder->SetBorderBackgroundColor(MarkedForDeletion ? FColorList::Red : FColorList::Black);
+	}
+	else
+	{
+		ChildSlot.DetachWidget();
+	}
+	return FReply::Handled();
+}
+
+
 #undef LOCTEXT_NAMESPACE
\ No newline at end of file
